[INTEGRATE] Changes to APIDocTool from 4.1 release branch.

diff --git a/Engine/Source/Programs/UnrealDocTool/APIDocTool/APIDocTool/APIConstant.cs b/Engine/Source/Programs/UnrealDocTool/APIDocTool/APIDocTool/APIConstant.cs
--- a/Engine/Source/Programs/UnrealDocTool/APIDocTool/APIDocTool/APIConstant.cs
+++ b/Engine/Source/Programs/UnrealDocTool/APIDocTool/APIDocTool/APIConstant.cs
@@ -13,12 +13,15 @@ namespace APIDocTool
 {
 	public abstract class APIConstant : APIMember
 	{
+		public readonly DoxygenEntity Entity;
+
 		public string BriefDescription;
 		public string FullDescription;
 
-		public APIConstant(APIPage InParent, string InName)
+		public APIConstant(APIPage InParent, DoxygenEntity InEntity, string InName)
 			: base(InParent, InName)
 		{
+			Entity = InEntity;
 		}
 
 		protected abstract void WriteDefinition(UdnWriter Writer);
@@ -55,6 +58,11 @@ namespace APIDocTool
 				}
 				Writer.LeaveTag("[/PARAM]");
 
+				// Write the references
+				Writer.EnterParam("references");
+				WriteReferencesSection(Writer, Entity);
+				Writer.LeaveParam();
+
 				Writer.LeaveTag("[/OBJECT]");
 			}
 		}
@@ -71,10 +79,10 @@ namespace APIDocTool
 		public readonly XmlNode ValueNode;
 		public string Initializer;
 
-		public APIConstantEnum(APIPage InParent, XmlNode InEnumNode, XmlNode InValueNode, string InName)
-			: base(InParent, InName)
+		public APIConstantEnum(APIPage InParent, DoxygenEntity InEnumEntity, XmlNode InValueNode, string InName)
+			: base(InParent, InEnumEntity, InName)
 		{
-			EnumNode = InEnumNode;
+			EnumNode = InEnumEntity.Node;
 			ValueNode = InValueNode;
 		}
 
@@ -100,7 +108,7 @@ namespace APIDocTool
 				foreach (XmlNode ValueNode in ValueNodeList)
 				{
 					string Name = ValueNode.SelectSingleNode("name").InnerText;
-					APIConstant Constant = new APIConstantEnum(Parent, EnumNode, ValueNode, Name);
+					APIConstant Constant = new APIConstantEnum(Parent, Entity, ValueNode, Name);
 					yield return Constant;
 				}
 			}
@@ -121,16 +129,14 @@ namespace APIDocTool
 
 	public class APIConstantVariable : APIConstant
 	{
-		public readonly DoxygenEntity Entity;
 		public readonly XmlNode Node;
 
 		public string Type;
 		public string Initializer;
 
 		public APIConstantVariable(APIPage InParent, DoxygenEntity InEntity)
-			: base(InParent, InEntity.Name)
+			: base(InParent, InEntity, InEntity.Name)
 		{
-			Entity = InEntity;
 			Node = Entity.Node;
 		}
 
diff --git a/Engine/Source/Programs/UnrealDocTool/APIDocTool/APIDocTool/APIDocTool.csproj b/Engine/Source/Programs/UnrealDocTool/APIDocTool/APIDocTool/APIDocTool.csproj
--- a/Engine/Source/Programs/UnrealDocTool/APIDocTool/APIDocTool/APIDocTool.csproj
+++ b/Engine/Source/Programs/UnrealDocTool/APIDocTool/APIDocTool/APIDocTool.csproj
@@ -74,6 +74,7 @@
     <Compile Include="APIMemberIndex.cs" />
     <Compile Include="APIPage.cs" />
     <Compile Include="APIEventParameters.cs" />
+    <Compile Include="APIGettingStarted.cs" />
     <Compile Include="APIRecord.cs" />
     <Compile Include="APITypeIndex.cs" />
     <Compile Include="BuildTarget.cs" />
diff --git a/Engine/Source/Programs/UnrealDocTool/APIDocTool/APIDocTool/APIEnum.cs b/Engine/Source/Programs/UnrealDocTool/APIDocTool/APIDocTool/APIEnum.cs
--- a/Engine/Source/Programs/UnrealDocTool/APIDocTool/APIDocTool/APIEnum.cs
+++ b/Engine/Source/Programs/UnrealDocTool/APIDocTool/APIDocTool/APIEnum.cs
@@ -13,6 +13,7 @@ namespace APIDocTool
 {
 	public class APIEnum : APIMember
     {
+		public readonly DoxygenEntity Entity;
 		public readonly XmlNode Node;
 		public APIProtection Protection;
 
@@ -26,15 +27,11 @@ namespace APIDocTool
 		public List<APIEnumValue> Values = new List<APIEnumValue>();
 
 		public APIEnum(APIPage InParent, DoxygenEntity InEntity, string InName)
-			: this(InParent, InEntity.Node, InName)
-		{
-		}
-
-		public APIEnum(APIPage InParent, XmlNode InNode, string InName)
 			: base(InParent, InName)
         {
 			// Set the defaults
-			Node = InNode;
+			Entity = InEntity;
+			Node = InEntity.Node;
 			Protection = ParseProtection(Node);
 
 			// Read the values
@@ -137,6 +134,11 @@ namespace APIDocTool
 				}
 				Writer.LeaveTag("[/PARAM]");
 
+				// Write the references
+				Writer.EnterTag("[PARAM:references]");
+				WriteReferencesSection(Writer, Entity);
+				Writer.LeaveTag("[/PARAM]");
+
 				Writer.LeaveTag("[/OBJECT]");
 			}
         }
diff --git a/Engine/Source/Programs/UnrealDocTool/APIDocTool/APIDocTool/APIFilter.cs b/Engine/Source/Programs/UnrealDocTool/APIDocTool/APIDocTool/APIFilter.cs
--- a/Engine/Source/Programs/UnrealDocTool/APIDocTool/APIDocTool/APIFilter.cs
+++ b/Engine/Source/Programs/UnrealDocTool/APIDocTool/APIDocTool/APIFilter.cs
@@ -88,7 +88,7 @@ namespace APIDocTool
 
 				// Write the class list
 				Writer.EnterTag("[PARAM:filters]");
-				APIFilter.WriteListSection(Writer, "filters", "Filters", Filters);
+				APIFilter.WriteListSection(Writer, "filters", "Filters", Filters.OrderBy(x => x.Name));
 				Writer.LeaveTag("[/PARAM]");
 
 				// Write the class list
diff --git a/Engine/Source/Programs/UnrealDocTool/APIDocTool/APIDocTool/APIFunction.cs b/Engine/Source/Programs/UnrealDocTool/APIDocTool/APIDocTool/APIFunction.cs
--- a/Engine/Source/Programs/UnrealDocTool/APIDocTool/APIDocTool/APIFunction.cs
+++ b/Engine/Source/Programs/UnrealDocTool/APIDocTool/APIDocTool/APIFunction.cs
@@ -278,6 +278,7 @@ namespace APIDocTool
     {
 		public static Dictionary<string, APIFunction> TemplateFunctions = new Dictionary<string, APIFunction>();
 
+		public readonly DoxygenEntity Entity;
 		public readonly XmlNode Node;
 		public readonly APIProtection Protection;
 		public readonly APIFunctionType FunctionType;
@@ -310,22 +311,19 @@ namespace APIDocTool
 		public List<string> Warnings = new List<string>();
 
 		public APIFunction(APIPage InParent, DoxygenEntity InEntity, APIFunctionKey InKey)
-			: this(InParent, InEntity.Node, InKey)
+			: this(InParent, InEntity, InKey, Utility.MakeLinkPath(InParent.LinkPath, InKey.GetLinkName()))
 		{
 		}
 
-		public APIFunction(APIPage InParent, XmlNode InNode, APIFunctionKey InKey)
-			: this(InParent, InNode, InKey, Utility.MakeLinkPath(InParent.LinkPath, InKey.GetLinkName()))
-		{
-		}
-
-		public APIFunction(APIPage InParent, XmlNode InNode, APIFunctionKey InKey, string InLinkPath)
-            : base(InParent, InKey.Name, InLinkPath)
+		public APIFunction(APIPage InParent, DoxygenEntity InEntity, APIFunctionKey InKey, string InLinkPath)
+            : base(InParent, InEntity.Name, InLinkPath)
         {
-			Node = InNode;
-			Protection = ParseProtection(InNode);
+			Entity = InEntity;
+			Node = Entity.Node;
+
+			Protection = ParseProtection(Node);
 			FunctionType = InKey.Type;
-			AddRefLink(Node.Attributes["id"].Value, this);
+			AddRefLink(Entity.Node.Attributes["id"].Value, this);
 
 			bIsTemplateSpecialization = Name.Contains('<');
 			if (Node.SelectSingleNode("templateparamlist") != null && !bIsTemplateSpecialization && !TemplateFunctions.ContainsKey(FullName))
@@ -459,6 +457,7 @@ namespace APIDocTool
 			}
 
 			// Parse the source lines
+			/*
 			XmlNode LocationNode = Node.SelectSingleNode("location");
 			if(LocationNode != null)
 			{
@@ -466,7 +465,7 @@ namespace APIDocTool
 				if(BodyFileAttribute != null)
 				{
 					string BodyFile = BodyFileAttribute.Value;
-					if(Program.IncludedSourceFiles.Contains(Path.GetFileName(BodyFile)))
+					if(!BodyFile.ToLowerInvariant().Split('\\', '/').Any(x => Program.ExcludeSourceDirectoriesHash.Contains(x)))
 					{
 						SourceFile File = SourceFileCache.Read(BodyFile);
 						if (File != null)
@@ -482,6 +481,7 @@ namespace APIDocTool
 					}
 				}
 			}
+			 */
 		}
 
 		public override void PostLink()
@@ -567,6 +567,61 @@ namespace APIDocTool
 			Writer.LeaveTag("[/REGION]");
 		}
 
+		private void WriteSourceSection(UdnWriter Writer)
+		{
+			if(Entity.BodyFile != null)
+			{
+				DoxygenSourceFile SourceFile = Entity.Module.FindSourceFile(Entity.BodyFile);
+				if(SourceFile != null)
+				{
+					int BodyStart = Math.Min(Math.Max(Entity.BodyStart - 1, 0), SourceFile.Lines.Count - 1);
+					int BodyEnd = Math.Min(Math.Max(Entity.BodyEnd, BodyStart), SourceFile.Lines.Count);
+					if(BodyEnd > BodyStart)
+					{
+						Writer.EnterSection("source", "Source");
+						Writer.EnterRegion("simplecode");
+
+						List<string> Lines = new List<string>();
+						int MinPrefix = int.MaxValue;
+
+						for (int LineIdx = BodyStart; LineIdx < BodyEnd; LineIdx++)
+						{
+							XmlNode Node = SourceFile.Lines[LineIdx];
+							string MarkdownLine = (Node == null)? "" : Markdown.ParseXmlCodeLine(Node, ResolveDoxygenLink);
+
+							int Prefix = 0;
+							while (Prefix < MarkdownLine.Length && MarkdownLine[Prefix] == ' ') Prefix++;
+
+							if(Prefix < MarkdownLine.Length && Prefix < MinPrefix)
+							{
+								MinPrefix = Prefix;
+							}
+
+							Lines.Add(MarkdownLine);
+						}
+
+						for (int Idx = 0; Idx < Lines.Count; Idx++)
+						{
+							int TextIdx = Math.Min(MinPrefix, Lines[Idx].Length);
+							if(TextIdx == Lines[Idx].Length)
+							{
+								Writer.Write("&nbsp;");
+							}
+							while(TextIdx < Lines[Idx].Length && Lines[Idx][TextIdx] == ' ')
+							{
+								Writer.Write("&nbsp;");
+								TextIdx++;
+							}
+							Writer.WriteLine(Lines[Idx].Substring(TextIdx) + "  ");
+						}
+
+						Writer.LeaveRegion();
+						Writer.LeaveSection();
+					}
+				}
+			}
+		}
+
 		private void WriteIcons(UdnWriter Writer)
 		{
 			Writer.WriteIcon(Icons.Function[(int)Protection]);
@@ -694,12 +749,20 @@ namespace APIDocTool
 
 				// Write the source
 				Writer.EnterTag("[PARAM:source]");
+				WriteSourceSection(Writer);
+/*
+
 				if (SourceLines != null)
 				{
 					Writer.EnterSection("source", "Source");
 					WriteSource(Writer);
 					Writer.LeaveSection();
 				}
+*/				Writer.LeaveTag("[/PARAM]");
+
+				// Write the reference info
+				Writer.EnterTag("[PARAM:references]");
+				WriteReferencesSection(Writer, Entity);
 				Writer.LeaveTag("[/PARAM]");
 
 				Writer.LeaveTag("[/OBJECT]");
diff --git a/Engine/Source/Programs/UnrealDocTool/APIDocTool/APIDocTool/APIFunctionGroup.cs b/Engine/Source/Programs/UnrealDocTool/APIDocTool/APIDocTool/APIFunctionGroup.cs
--- a/Engine/Source/Programs/UnrealDocTool/APIDocTool/APIDocTool/APIFunctionGroup.cs
+++ b/Engine/Source/Programs/UnrealDocTool/APIDocTool/APIDocTool/APIFunctionGroup.cs
@@ -17,24 +17,19 @@ namespace APIDocTool
 		public readonly APIFunctionType FunctionType;
 
 		public APIFunctionGroup(APIPage InParent, APIFunctionKey InKey, IEnumerable<DoxygenEntity> InEntities) 
-			: this(InParent, InKey, InEntities.Select(x => x.Node))
-		{
-		}
-
-		public APIFunctionGroup(APIPage InParent, APIFunctionKey InKey, IEnumerable<XmlNode> InNodes) 
 			: base(InParent, InKey.Name, Utility.MakeLinkPath(InParent.LinkPath, InKey.GetLinkName()))
 		{
 			FunctionType = InKey.Type;
 
 			// Build a list of prototypes for each function node, to be used for sorting
-			List<XmlNode> Nodes = new List<XmlNode>(InNodes.OrderBy(x => x.SelectNodes("param").Count));
+			List<DoxygenEntity> Entities = new List<DoxygenEntity>(InEntities.OrderBy(x => x.Node.SelectNodes("param").Count));
 
 			// Create the functions
 			int LinkIndex = 1;
-			foreach (XmlNode Node in Nodes)
+			foreach(DoxygenEntity Entity in Entities)
 			{
 				string NewLinkPath = Utility.MakeLinkPath(LinkPath, String.Format("{0}", LinkIndex));
-				APIFunction Function = new APIFunction(this, Node, InKey, NewLinkPath);
+				APIFunction Function = new APIFunction(this, Entity, InKey, NewLinkPath);
 				Children.Add(Function);
 				LinkIndex++;
 			}
diff --git a/Engine/Source/Programs/UnrealDocTool/APIDocTool/APIDocTool/APIGettingStarted.cs b/Engine/Source/Programs/UnrealDocTool/APIDocTool/APIDocTool/APIGettingStarted.cs
new file mode 100644
--- /dev/null
+++ b/Engine/Source/Programs/UnrealDocTool/APIDocTool/APIDocTool/APIGettingStarted.cs
@@ -0,0 +1,83 @@
+﻿using System;
+using System.Collections.Generic;
+using System.Linq;
+using System.Text;
+using System.Threading.Tasks;
+
+namespace APIDocTool
+{
+	class APIGettingStarted : APIPage
+	{
+		public APIGettingStarted(APIPage InParent)
+			: base(InParent, "QuickStart")
+		{
+		}
+
+		public void WriteSitemapContents(string OutputPath)
+		{
+		}
+
+		public override void GatherReferencedPages(List<APIPage> ReferencedPages)
+		{
+		}
+
+		public override void AddToManifest(UdnManifest Manifest)
+		{
+		}
+
+		public override void WritePage(UdnManifest Manifest, string OutputPath)
+		{
+			using (UdnWriter Writer = new UdnWriter(OutputPath))
+			{
+				Writer.WritePageHeader("Getting started with the Unreal Engine API", PageCrumbs, "Getting started with the Unreal Engine API");
+
+				Writer.WriteHeading(2, "Orientation");
+
+				Writer.WriteLine("Games, programs and the Unreal Editor are all *targets* built by UnrealBuildTool. Each target is compiled from C++ *modules*, each implementing a particular area of functionality. Your game is a target, and your game code is implemented in one or more modules.");
+				Writer.WriteLine();
+				Writer.WriteLine("Code in each module can use other modules by referencing them in their *build rules*. Build rules for each module are given by C# scripts with the .build.cs extension. *Target rules* are given by C# scripts with the .target.cs extension.");
+				Writer.WriteLine();
+				Writer.WriteLine(Manifest, "Modules supplied with the Unreal Engine are divided into three categories; the {Runtime|ModuleIndex:Runtime Modules}, functionality for the {Editor|ModuleIndex:Editor Modules}, and {Developer|ModuleIndex:Developer Modules} utilities.");
+				Writer.WriteLine();
+				Writer.WriteLine(Manifest, "Most gameplay programming just uses runtime modules, and the three most commonly encountered are {Core|Filter:Core}, {CoreUObject|Filter:CoreUObject} and {Engine|Filter:Engine}.");
+				Writer.WriteLine();
+
+				Writer.WriteHeading(2, "Core");
+
+				Writer.WriteLine(Manifest, "The **{Core|Filter:Core}** module provides a common framework for Unreal modules to communicate; a standard set of types, a {math library|Filter:Core.Math}, a {container|Filter:Core.Containers} library, and a lot of the hardware abstraction that allows Unreal to run on so many platforms.");
+				Writer.WriteLine(Manifest, "Some common topics are listed below; full documentation is available {here|Filter:Core}.");
+				Writer.WriteLine();
+				List<UdnListItem> CoreItems = new List<UdnListItem>();
+				CoreItems.Add(new UdnListItem("Basic types:", Manifest.FormatString("bool &middot; float/double &middot; {int8}/{int16}/{int32}/{int64} &middot; {uint8}/{uint16}/{uint32}/{uint64} &middot; {ANSICHAR} &middot; {TCHAR} &middot; {FString}"), null));
+				CoreItems.Add(new UdnListItem("Math:", Manifest.FormatString("{FMath} &middot; {FVector} &middot; {FRotator} &middot; {FTransform} &middot; {FMatrix} &middot; {More...|Filter:Core.Math}"), null));
+				CoreItems.Add(new UdnListItem("Containers:", Manifest.FormatString("{TArray} &middot; {TList} &middot; {TMap} &middot; {More...|Filter:Core.Containers}"), null));
+				CoreItems.Add(new UdnListItem("Other:", Manifest.FormatString("{FName} &middot; {FArchive} &middot; {FOutputDevice}"), null));
+				Writer.WriteList(CoreItems);
+				
+				Writer.WriteHeading(2, "CoreUObject");
+
+				Writer.WriteLine(Manifest, "The **{CoreUObject|Filter:CoreUObject}** module defines UObject, the base class for all managed objects in Unreal. Managed objects are key to integrating with the editor, for serialization, network replication, and runtime type information. UObject derived classes are garbage collected.");
+				Writer.WriteLine(Manifest, "Some common topics are listed below; full documentation is available {here|Filter:CoreUObject}.");
+				Writer.WriteLine();
+				List<UdnListItem> CoreUObjectItems = new List<UdnListItem>();
+				CoreUObjectItems.Add(new UdnListItem("Classes:", Manifest.FormatString("{UObject} &middot; {UClass} &middot; {UProperty} &middot; {UPackage}"), null));
+				CoreUObjectItems.Add(new UdnListItem("Functions:", Manifest.FormatString("{ConstructObject} &middot; {FindObject} &middot; {Cast} &middot; {CastChecked}"), null));
+				Writer.WriteList(CoreUObjectItems);
+
+				Writer.WriteHeading(2, "Engine");
+
+				Writer.WriteLine(Manifest, "The **{Engine|Filter:Engine}** module contains functionality you’d associate with a game. The game world, actors, characters, physics and special effects are all defined here.");
+				Writer.WriteLine(Manifest, "Some common topics are listed below; full documentation is available {here|Filter:Engine}.");
+				Writer.WriteLine();
+				List<UdnListItem> EngineItems = new List<UdnListItem>();
+				EngineItems.Add(new UdnListItem("Actors:", Manifest.FormatString("{AActor} &middot; {AVolume} &middot; {AGameMode} &middot; {AHUD} &middot; {More...|Hierarchy:AActor}"), null));
+				EngineItems.Add(new UdnListItem("Pawns:", Manifest.FormatString("{APawn} &middot; {ACharacter} &middot; {AWheeledVehicle}"), null));
+				EngineItems.Add(new UdnListItem("Controllers:", Manifest.FormatString("{AController} &middot; {AAIController} &middot; {APlayerController}"), null));
+				EngineItems.Add(new UdnListItem("Components:", Manifest.FormatString("{UActorComponent} &middot; {UBrainComponent} &middot; {UInputComponent} &middot; {USkeletalMeshComponent} &middot; {UParticleSystemComponent} &middot; {More...|Hierarchy:UActorComponent}"), null));
+				EngineItems.Add(new UdnListItem("Gameplay:", Manifest.FormatString("{UPlayer} &middot; {ULocalPlayer} &middot; {UWorld} &middot; {ULevel}"), null));
+				EngineItems.Add(new UdnListItem("Assets:", Manifest.FormatString("{UTexture} &middot; {UMaterial} &middot; {UStaticMesh} &middot; {USkeletalMesh} &middot; {UParticleSystem}"), null));
+				Writer.WriteList(EngineItems);
+			}
+		}
+	}
+}
diff --git a/Engine/Source/Programs/UnrealDocTool/APIDocTool/APIDocTool/APIIndex.cs b/Engine/Source/Programs/UnrealDocTool/APIDocTool/APIDocTool/APIIndex.cs
--- a/Engine/Source/Programs/UnrealDocTool/APIDocTool/APIDocTool/APIIndex.cs
+++ b/Engine/Source/Programs/UnrealDocTool/APIDocTool/APIDocTool/APIIndex.cs
@@ -14,22 +14,37 @@ namespace APIDocTool
 {
 	class APIIndex : APIPage
 	{
-		public static string[] QuickLinkPaths = 
-		{
-			"API\\Runtime\\CoreUObject\\UObject\\UObject",
-			"API\\Runtime\\Engine\\GameFramework\\AActor",
-			"API\\Runtime\\Engine\\Components\\UActorComponent",
-		};
-
 		public List<APIPage> Pages = new List<APIPage>();
-		public List<APIModuleIndex> ModuleIndexes = new List<APIModuleIndex>();
+		public List<APIModuleIndex> ChildModuleIndexes = new List<APIModuleIndex>();
 
+		public APIGettingStarted GettingStarted;
+		public APIModuleIndex ModuleIndex;
 		public APIConstantIndex ConstantIndex;
 		public APIFunctionIndex FunctionIndex;
 		public APITypeIndex TypeIndex;
 		public APIHierarchy RecordHierarchy;
 
-		public List<APIMember> QuickLinks = new List<APIMember>();
+		public static APIModuleCategory FindRootCategoryByPath(List<APIModuleCategory> Categories, string InPath)
+		{
+			const string EngineSourcePath = "Engine/Source/";
+			string NormalizedSrcPath = Path.GetFullPath(InPath).Replace('\\', '/').TrimEnd('/');
+
+			int CategoryMinIdx = NormalizedSrcPath.LastIndexOf(EngineSourcePath);
+			if (CategoryMinIdx != -1)
+			{
+				CategoryMinIdx += EngineSourcePath.Length;
+
+				int CategoryMaxIdx = CategoryMinIdx;
+				while (CategoryMaxIdx < NormalizedSrcPath.Length && NormalizedSrcPath[CategoryMaxIdx] != '/') CategoryMaxIdx++;
+
+				string CategoryName = NormalizedSrcPath.Substring(CategoryMinIdx, CategoryMaxIdx - CategoryMinIdx) + " Modules";
+
+				APIModuleCategory Category = Categories.FirstOrDefault(x => x.Name.ToLower() == CategoryName.ToLower());
+				if(Category != null) return Category;
+			}
+
+			return Categories.Last();
+		}
 
 		public APIIndex(List<DoxygenModule> InModules)
 			: base(null, "API")
@@ -41,21 +56,42 @@ namespace APIDocTool
 				NameToModuleMap.Add(Module.Name, Module);
 			}
 
-			// Create a module category tree
-			APIModuleCategory RootCategory = new APIModuleCategory(null);
-			RootCategory.AddModules(InModules.Select(x => new KeyValuePair<string, string>(x.Name, x.BaseSrcDir)).ToArray());
-			Debug.Assert(RootCategory.MinorModules.Count == 0 && RootCategory.MajorModules.Count == 0);
+			// Add some default categories to fix the order
+			APIModuleCategory RootCategory = new APIModuleCategory("Modules");
+			RootCategory.Categories.Add(new APIModuleCategory("Runtime Modules"));
+			RootCategory.Categories.Add(new APIModuleCategory("Editor Modules"));
+			RootCategory.Categories.Add(new APIModuleCategory("Developer Modules"));
+
+			// Add all of the modules to a category
+			Dictionary<DoxygenModule, APIModuleCategory> ModuleToCategory = new Dictionary<DoxygenModule, APIModuleCategory>();
+			foreach(DoxygenModule Module in InModules)
+			{
+				APIModuleCategory Category = FindRootCategoryByPath(RootCategory.Categories, Module.BaseSrcDir);
+				ModuleToCategory.Add(Module, Category);
+			}
 
 			// Create all the index pages
 			foreach (APIModuleCategory Category in RootCategory.Categories)
 			{
-				if (!Category.IsEmpty)
-				{
-					APIModuleIndex ModuleIndex = new APIModuleIndex(this, Category, InModules);
-					ModuleIndexes.Add(ModuleIndex);
-				}
+				APIModuleIndex ChildModuleIndex = new APIModuleIndex(this, Category);
+				ChildModuleIndexes.Add(ChildModuleIndex);
+			}
+			Pages.AddRange(ChildModuleIndexes);
+
+			// Populate all the categories and create all the module pages
+			foreach(KeyValuePair<DoxygenModule, APIModuleCategory> Pair in ModuleToCategory)
+			{
+				APIModuleIndex Parent = ChildModuleIndexes.First(x => x.Category == Pair.Value);
+
+				APIModule Module = APIModule.Build(Parent, Pair.Key);
+				Parent.Children.Add(Module);
+
+				Pair.Value.AddModule(Module);
 			}
-			Pages.AddRange(ModuleIndexes);
+
+			// Expand the Core section by default
+			APIModuleCategory DefaultCategory = RootCategory.Categories[0].Categories.FirstOrDefault(x => x.Name == "Core Runtime Modules");
+			if (DefaultCategory != null) DefaultCategory.Expanded = true;
 
 			// Get all the members that were created as part of building the modules. After this point we'll only create index pages.
 			List<APIMember> AllMembers = new List<APIMember>(GatherPages().OfType<APIMember>().OrderBy(x => x.FullName));
@@ -68,6 +104,14 @@ namespace APIDocTool
 				Member.PostLink();
 			}
 
+			// Create the quick start page
+			GettingStarted = new APIGettingStarted(this);
+			Pages.Add(GettingStarted);
+
+			// Create an index of all the modules
+			ModuleIndex = new APIModuleIndex(this, RootCategory);
+			Pages.Add(ModuleIndex);
+
 			// Create an index of all the constants
 			ConstantIndex = new APIConstantIndex(this, AllMembers);
 			Pages.Add(ConstantIndex);
@@ -83,23 +127,18 @@ namespace APIDocTool
 			// Create an index of all the classes
 			RecordHierarchy = APIHierarchy.Build(this, AllMembers.OfType<APIRecord>());
 			Pages.Add(RecordHierarchy);
-
-			// Create all the quick links
-			foreach (string QuickLinkPath in QuickLinkPaths)
-			{
-				APIMember Member = AllMembers.FirstOrDefault(x => x.LinkPath == QuickLinkPath);
-				if (Member != null) QuickLinks.Add(Member);
-			}
 		}
 
 		public void WriteSitemapContents(string OutputPath)
 		{
 			List<SitemapNode> RootNodes = new List<SitemapNode>();
-			RootNodes.Add(new SitemapNode("Introduction", SitemapLinkPath));
-			foreach (APIModuleIndex ModuleIndex in ModuleIndexes)
+			RootNodes.Add(new SitemapNode("Contents", SitemapLinkPath));
+			RootNodes.Add(new SitemapNode("Getting started", GettingStarted.SitemapLinkPath));
+			foreach (APIModuleIndex ChildModuleIndex in ChildModuleIndexes)
 			{
-				RootNodes.Add(ModuleIndex.CreateSitemapNode());
+				RootNodes.Add(ChildModuleIndex.CreateSitemapNode());
 			}
+			RootNodes.Add(new SitemapNode("All modules", ModuleIndex.SitemapLinkPath));
 			RootNodes.Add(new SitemapNode("All constants", ConstantIndex.SitemapLinkPath));
 			RootNodes.Add(new SitemapNode("All functions", FunctionIndex.SitemapLinkPath));
 			RootNodes.Add(new SitemapNode("All types", TypeIndex.SitemapLinkPath));
@@ -152,7 +191,7 @@ namespace APIDocTool
 
 		public override void GatherReferencedPages(List<APIPage> ReferencedPages)
 		{
-			ReferencedPages.AddRange(ModuleIndexes);
+			ReferencedPages.AddRange(ChildModuleIndexes);
 			ReferencedPages.AddRange(Pages);
 		}
 
@@ -164,74 +203,6 @@ namespace APIDocTool
 			}
 		}
 
-		public string GetModuleIndexLink(string Name)
-		{
-			foreach(APIModuleIndex ModuleIndex in ModuleIndexes)
-			{
-				if (ModuleIndex.Name == Name)
-				{
-					return String.Format("[{0}]({1})", Name, ModuleIndex.LinkPath);
-				}
-			}
-			return Name;
-		}
-
-		public string MakeLinkString(Dictionary<string, List<APIPage>> Manifest, string Text)
-		{
-			StringBuilder Output = new StringBuilder();
-			for (int Idx = 0; Idx < Text.Length; Idx++)
-			{
-				if(Text[Idx] == '{')
-				{
-					if(Text[Idx + 1] == '{')
-					{
-						Output.Append(Text[++Idx]);
-					}
-					else
-					{
-						int EndIdx = Text.IndexOf('}', Idx + 1);
-						string LinkText = Text.Substring(Idx + 1, EndIdx - Idx - 1);
-						string LinkName = LinkText;
-
-						int LinkNameIdx = LinkText.IndexOf('|');
-						if(LinkNameIdx != -1)
-						{
-							LinkName = LinkText.Substring(LinkNameIdx + 1);
-							LinkText = LinkText.Substring(0, LinkNameIdx);
-						}
-					
-						List<APIPage> LinkPages;
-						if (Manifest.TryGetValue(LinkName, out LinkPages))
-						{
-							Output.AppendFormat("[{0}]({1})", LinkText, LinkPages[0].LinkPath);
-						}
-						else
-						{
-							Output.Append(LinkText);
-						}
-						Idx = EndIdx;
-					}
-				}
-				else
-				{
-					Output.Append(Text[Idx]);
-				}
-			}
-			return Output.ToString();
-		}
-
-		public void WriteManifestLinks(UdnWriter Writer, Dictionary<string, List<APIPage>> Manifest, string Text)
-		{
-			Writer.WriteLine(MakeLinkString(Manifest, Text));
-		}
-
-		public void WriteIndexLinks(UdnWriter Writer, Dictionary<string, List<APIPage>> Manifest, string Text)
-		{
-			Writer.EnterTag("[REGION:memberindexlinks]");
-			Writer.WriteLine(MakeLinkString(Manifest, Text));
-			Writer.LeaveTag("[/REGION]");
-		}
-
 		public override void WritePage(UdnManifest Manifest, string OutputPath)
 		{
 			using (UdnWriter Writer = new UdnWriter(OutputPath))
@@ -243,48 +214,25 @@ namespace APIDocTool
 				Writer.WriteLine("The API reference is an early work in progress, and some information may be missing or out of date. It serves mainly as a low level index of Engine classes and functions.");
 				Writer.WriteLine("For tutorials, walkthroughs and detailed guides to programming with Unreal, please see the [Unreal Engine Programming](http://docs.unrealengine.com/latest/INT/Programming/index.html) home on the web.");
 
-				Writer.WriteHeading(2, "Orientation");
+				Writer.WriteHeading(2, "Introduction");
 
-				Writer.WriteLine("Games, programs and the Unreal Editor are all *targets* built by UnrealBuildTool. Each target is compiled from C++ *modules*, each implementing a particular area of functionality. Your game is a target, and your game code is implemented in one or more modules.");
+				Writer.WriteLine(Manifest, "The Unreal Engine consists of a large number of modules divided into three categories: the {Runtime|ModuleIndex:Runtime Modules}, functionality for the {Editor|ModuleIndex:Editor Modules}, and {Developer|ModuleIndex:Developer Modules} utilities. Full documentation for each of these categories is below.");
 				Writer.WriteLine();
-				Writer.WriteLine("Code in each module can use other modules by referencing them in their *build rules*. Build rules for each module are given by C# scripts with the .build.cs extension. *Target rules* are given by C# scripts with the .target.cs extension.");
-				Writer.WriteLine();
-				Writer.WriteLine(Manifest, "The Unreal Engine consists of a large number of modules divided into three categories: the common {Runtime|ModuleIndex:Runtime}, {Editor|ModuleIndex:Editor} specific functionality, and {Developer|ModuleIndex:Developer} utilities. Three of the most important runtime modules are {Core|Filter:Core}, {CoreUObject|Filter:CoreUObject} and {Engine|Filter:Engine}.");
-
-				Writer.WriteHeading(2, "Core");
+				Writer.WriteLine("To explore the API from some of the most frequently encountered Unreal concepts and types, see the API __[getting started](" + GettingStarted.LinkPath + ")__ page.");
 
-				Writer.WriteLine(Manifest, "The **{Core|Filter:Core}** module provides a common framework for Unreal modules to communicate; a standard set of types, a {math library|Filter:Core.Math}, a {container|Filter:Core.Containers} library, and a lot of the hardware abstraction that allows Unreal to run on so many platforms.");
-				Writer.WriteLine();
-				Writer.WriteLine("| Topic | Links |");
-				Writer.WriteLine("| ----- | ----- |");
-				Writer.WriteLine(Manifest, "|Basic types:| bool &middot; float/double &middot; {int8}/{int16}/{int32}/{int64} &middot; {uint8}/{uint16}/{uint32}/{uint64} &middot; {ANSICHAR} &middot; {TCHAR} &middot; {FString}|");
-				Writer.WriteLine(Manifest, "|Math:| {FMath} &middot; {FVector} &middot; {FRotator} &middot; {FTransform} &middot; {FMatrix} &middot; {More...|Filter:Core.Math} |");
-				Writer.WriteLine(Manifest, "|Containers:| {TArray} &middot; {TList} &middot; {TMap} &middot; {More...|Filter:Core.Containers} |");
-				Writer.WriteLine(Manifest, "|Other:| {FName} &middot; {FArchive} &middot; {FOutputDevice} |");
-				Writer.WriteLine();
+				Writer.WriteHeading(2, "Contents");
 
-				Writer.WriteHeading(2, "CoreUObject");
+				Writer.EnterRegion("memberindexlinks");
+				Writer.WriteLine("[Getting started]({0}) &middot; [All modules]({1}) &middot; [All constants]({2}) &middot; [All functions]({3}) &middot; [All types]({4}) &middot; [Class hierarchy]({5})", GettingStarted.LinkPath, ModuleIndex.LinkPath, ConstantIndex.LinkPath, FunctionIndex.LinkPath, TypeIndex.LinkPath, RecordHierarchy.LinkPath);
+				Writer.LeaveRegion();
 
-				Writer.WriteLine(Manifest, "The **{CoreUObject|Filter:CoreUObject}** module defines UObject, the base class for all managed objects in Unreal. Managed objects are key to integrating with the editor, for serialization, network replication, and runtime type information. UObject derived classes are garbage collected.");
-				Writer.WriteLine();
-				Writer.WriteLine("| Topic | Links |");
-				Writer.WriteLine("| ----- | ----- |");
-				Writer.WriteLine(Manifest, "|Classes:|{UObject} &middot; {UClass} &middot; {UProperty} &middot; {UPackage}|");
-				Writer.WriteLine(Manifest, "|Functions:|{ConstructObject} &middot; {FindObject} &middot; {Cast} &middot; {CastChecked}|");
-				Writer.WriteLine();
-
-				Writer.WriteHeading(2, "Engine");
-
-				Writer.WriteLine(Manifest, "The **{Engine|Filter:Engine}** module contains functionality you’d associate with a game. The game world, actors, characters, physics and special effects are all defined here.");
-				Writer.WriteLine();
-				Writer.WriteLine("| Topic | Links |");
-				Writer.WriteLine("| ----- | ----- |");
-				Writer.WriteLine(Manifest, "|Actors:|{AActor} &middot; {AVolume} &middot; {AGameMode} &middot; {AHUD} &middot; {More...|Hierarchy:AActor}");
-				Writer.WriteLine(Manifest, "|Pawns:|{APawn} &middot; {ACharacter} &middot; {AWheeledVehicle}");
-				Writer.WriteLine(Manifest, "|Controllers:|{AController} &middot; {AAIController} &middot; {APlayerController}"); 
-				Writer.WriteLine(Manifest, "|Components:|{UActorComponent} &middot; {UBrainComponent} &middot; {UInputComponent} &middot; {USkeletalMeshComponent} &middot; {UParticleSystemComponent} &middot; {More...|Hierarchy:UActorComponent}|");
-				Writer.WriteLine(Manifest, "|Gameplay:|{UPlayer} &middot; {ULocalPlayer} &middot; {UWorld} &middot; {ULevel}");
-				Writer.WriteLine(Manifest, "|Assets:|{UTexture} &middot; {UMaterial} &middot; {UStaticMesh} &middot; {USkeletalMesh} &middot; {UParticleSystem}|");
+				foreach(APIModuleIndex ChildModuleIndex in ChildModuleIndexes)
+				{
+					Writer.WriteHeading(2, ChildModuleIndex.Name);
+					Writer.EnterRegion("modules-list");
+					ChildModuleIndex.WriteExpandableModuleList(Writer, ChildModuleIndex.Name.Replace(' ', '_'));
+					Writer.LeaveRegion();
+				}
 
 				Writer.WriteLine("<br>");
 			}
diff --git a/Engine/Source/Programs/UnrealDocTool/APIDocTool/APIDocTool/APIMember.cs b/Engine/Source/Programs/UnrealDocTool/APIDocTool/APIDocTool/APIMember.cs
--- a/Engine/Source/Programs/UnrealDocTool/APIDocTool/APIDocTool/APIMember.cs
+++ b/Engine/Source/Programs/UnrealDocTool/APIDocTool/APIDocTool/APIMember.cs
@@ -78,7 +78,7 @@ namespace APIDocTool
 				}
 				else if (Entity.Kind == "typedef")
 				{
-					Children.Add(new APITypeDef(Parent, Entity.Node));
+					Children.Add(new APITypeDef(Parent, Entity));
 				}
 				else if (Entity.Kind == "enum")
 				{
@@ -266,6 +266,64 @@ namespace APIDocTool
 			Writer.LeaveTag("[/REGION]");
 		}
 
+		public void WriteReferencesSection(UdnWriter Writer, DoxygenEntity Entity)
+		{
+			List<UdnListItem> Items = new List<UdnListItem>();
+
+			// Get the module
+			for (APIPage Page = this; Page != null; Page = Page.Parent)
+			{
+				APIModule Module = Page as APIModule;
+				if(Module != null)
+				{
+					Items.Add(new UdnListItem("Module", String.Format("[{0}]({1})", Module.Name, Module.LinkPath), null));
+					break;
+				}
+			}
+
+			// Get the header file
+			if (Entity.File != null)
+			{
+				string NormalizedFileName = GetNormalizedFileName(Entity.File);
+				if (NormalizedFileName != null)
+				{
+					Items.Add(new UdnListItem("Header", Utility.EscapeText(NormalizedFileName), null));
+				}
+			}
+
+			// Get the source file
+			if(Entity.BodyFile != null && Entity.BodyFile != Entity.File)
+			{
+				string NormalizedFileName = GetNormalizedFileName(Entity.BodyFile);
+				if (NormalizedFileName != null)
+				{
+					Items.Add(new UdnListItem("Source", Utility.EscapeText(NormalizedFileName), null));
+				}
+			}
+
+			// Write the section
+			if(Items.Count > 0)
+			{
+				Writer.EnterSection("references", "References");
+				Writer.WriteList(Items);
+				Writer.LeaveSection();
+			}
+		}
+
+		public static string GetNormalizedFileName(string InFileName)
+		{
+			string NormalizedFileName = InFileName.Replace('\\', '/');
+			const string EngineSource = "/Engine/Source/";
+
+			int SourceIndex = NormalizedFileName.IndexOf(EngineSource);
+			if (SourceIndex != -1)
+			{
+				return NormalizedFileName.Substring(SourceIndex + EngineSource.Length);
+			}
+
+			return null;
+		}
+
 		public static string AnnotateAnonymousNames(string Name)
 		{
 			string[] Sections = Name.Split(new string[]{ "::" }, StringSplitOptions.None);
diff --git a/Engine/Source/Programs/UnrealDocTool/APIDocTool/APIDocTool/APIModuleIndex.cs b/Engine/Source/Programs/UnrealDocTool/APIDocTool/APIDocTool/APIModuleIndex.cs
--- a/Engine/Source/Programs/UnrealDocTool/APIDocTool/APIDocTool/APIModuleIndex.cs
+++ b/Engine/Source/Programs/UnrealDocTool/APIDocTool/APIDocTool/APIModuleIndex.cs
@@ -13,104 +13,24 @@ namespace APIDocTool
 	class APIModuleCategory
 	{
 		public string Name;
-		public List<string> MajorModules = new List<string>();
-		public List<string> MinorModules = new List<string>();
+		public List<APIModule> Modules = new List<APIModule>();
 		public List<APIModuleCategory> Categories = new List<APIModuleCategory>();
+		public bool Expanded = false;
 
 		public APIModuleCategory(string InName)
 		{
 			Name = InName;
 		}
 
-		public void GetModules(List<string> Modules)
-		{
-			Modules.AddRange(MajorModules);
-			Modules.AddRange(MinorModules);
-
-			foreach (APIModuleCategory Category in Categories)
-			{
-				Category.GetModules(Modules);
-			}
-		}
-
-		public void AddModuleByPath(string Name, string BaseSrcDir, bool bIsMajorModule)
+		public void AddModule(APIModule Module)
 		{
 			// Find the description and category
-			string ModuleSettingsPath = "Module." + Name;
+			string ModuleSettingsPath = "Module." + Module.Name;
 			string CategoryText = Program.Settings.FindValueOrDefault(ModuleSettingsPath + ".Category", null);
 
-			// Get a default category based on the module directory if there's not one set
-			if (CategoryText == null)
-			{
-				const string EngineSourcePath = "Engine/Source/";
-				string NormalizedSrcPath = Path.GetFullPath(BaseSrcDir).Replace('\\', '/').TrimEnd('/');
-
-				int CategoryMinIdx = NormalizedSrcPath.IndexOf(EngineSourcePath);
-				if(CategoryMinIdx != -1) CategoryMinIdx += EngineSourcePath.Length;
-					
-				int CategoryMaxIdx = NormalizedSrcPath.LastIndexOf('/');
-				if (CategoryMinIdx >= 0 && CategoryMaxIdx > CategoryMinIdx)
-				{
-					CategoryText = NormalizedSrcPath.Substring(CategoryMinIdx, CategoryMaxIdx - CategoryMinIdx).Replace('/', '|');
-				}
-				else
-				{
-					CategoryText = "Other";
-				}
-			}
-
 			// Build a module from all the members
-			AddModuleByCategory(Name, CategoryText, bIsMajorModule);
-		}
-		
-		public void AddModuleByCategory(string Name, string CategoryText, bool bIsMajorModule)
-		{
-			APIModuleCategory Category = AddCategory(CategoryText);
-			if (bIsMajorModule)
-			{
-				Category.MajorModules.Add(Name);
-			}
-			else
-			{
-				Category.MinorModules.Add(Name);
-			}
-		}
-
-		public void AddModules(KeyValuePair<string, string>[] NamePathList)
-		{
-			// Read the layout settings from the ini file
-			string[] CategoryOrder = Program.Settings.FindValueOrDefault("Index.CategoryOrder", "").Split('\n');
-			string[] ModuleOrder = Program.Settings.FindValueOrDefault("Index.ModuleOrder", "").Split('\n');
-			HashSet<string> MajorModules = new HashSet<string>(Program.Settings.FindValueOrDefault("Index.MajorModules", "").Split('\n'));
-
-			// Add all the categories specified in the order
-			foreach (string CategoryName in CategoryOrder)
-			{
-				AddCategory(CategoryName);
-			}
-
-			// Add all the module names whose order is specified
-			foreach (string ModuleName in ModuleOrder)
-			{
-				for(int Idx = 0; Idx < NamePathList.Length; Idx++)
-				{
-					if (NamePathList[Idx].Key == ModuleName)
-					{
-						AddModuleByPath(ModuleName, NamePathList[Idx].Value, MajorModules.Contains(ModuleName));
-						break;
-					}
-				}
-			}
-
-			// Add all the other modules to the tree
-			HashSet<string> AddedModules = new HashSet<string>(ModuleOrder);
-			foreach(KeyValuePair<string, string> NamePathPair in NamePathList)
-			{
-				if (!AddedModules.Contains(NamePathPair.Key))
-				{
-					AddModuleByPath(NamePathPair.Key, NamePathPair.Value, MajorModules.Contains(NamePathPair.Key));
-				}
-			}
+			APIModuleCategory Category = String.IsNullOrEmpty(CategoryText)? this : AddCategory(CategoryText);
+			Category.Modules.Add(Module);
 		}
 
 		public APIModuleCategory AddCategory(string Path)
@@ -143,37 +63,78 @@ namespace APIDocTool
 
 		public bool IsEmpty
 		{
-			get { return MajorModules.Count == 0 && MinorModules.Count == 0 && Categories.All(x => x.IsEmpty); }
+			get { return Modules.Count == 0 && Categories.All(x => x.IsEmpty); }
 		}
 
-		public void Write(UdnWriter Writer, int Depth, Dictionary<string, APIModule> Modules)
+		public void Write(UdnWriter Writer, int Depth)
 		{
 			if (!IsEmpty)
 			{
-				// Find all the major modules in this category
-				if (MajorModules.Count > 0)
-				{
-					Writer.WriteList("Name", "Description", MajorModules.Select(x => Modules[x].GetListItem()));
-				}
+				// CSS region for indenting 
+				Writer.EnterRegion("module-sections-list");
 
-				// Write all the minor modules in this category
-				if (MinorModules.Count > 0)
+				// Find all the modules in this category
+				if(Modules.Count > 0)
 				{
-					Writer.EnterRegion("syntax");
-					Writer.WriteFilterList(MinorModules.Select(x => Modules[x].GetFilterListItem()).ToArray());
-					Writer.LeaveRegion();
+					Writer.WriteList(Modules.OrderBy(x => x.Name).Select(x => x.GetListItem()));
 				}
 
 				// Write all the subcategories
-				foreach (APIModuleCategory Category in Categories)
+				foreach (APIModuleCategory Category in Categories.OrderBy(x => x.Name))
 				{
 					Writer.WriteHeading(Depth, Category.Name);
 					Writer.WriteLine();
-					Category.Write(Writer, Depth + 1, Modules);
+					Category.Write(Writer, Depth + 1);
 				}
+
+				// End of CSS region
+				Writer.LeaveRegion();
 			}
 		}
 
+		public void WriteExpandable(UdnWriter Writer, string UniqueId)
+		{
+			// Top level modules
+			if(Modules.Count > 0)
+			{
+				Writer.WriteList(Modules.OrderBy(x => x.Name).Select(x => x.GetListItem()));
+			}
+
+			// CSS region for indenting 
+//			Writer.EnterRegion("module-sections-list");
+
+			// Expandable section for each subcategory
+			int UniqueIdSuffix = 1;
+			foreach (APIModuleCategory Category in Categories.OrderBy(x => x.Name))
+			{
+				if(!Category.IsEmpty)
+				{
+					string NewUniqueId = String.Format("{0}_{1}", UniqueId, UniqueIdSuffix);
+
+					Writer.EnterTag(Category.Expanded? "[OBJECT:ModuleSectionExpanded]" : "[OBJECT:ModuleSection]");
+
+					Writer.WriteLine("[PARAMLITERAL:id]");
+					Writer.WriteEscapedLine(NewUniqueId);
+					Writer.WriteLine("[/PARAMLITERAL]");
+
+					Writer.WriteLine("[PARAM:heading]");
+					Writer.WriteEscapedLine(Category.Name);
+					Writer.WriteLine("[/PARAM]");
+
+					Writer.EnterTag("[PARAM:content]");
+					Category.WriteExpandable(Writer, NewUniqueId);
+					Writer.LeaveTag("[/PARAM]");
+
+					Writer.LeaveTag("[/OBJECT]");
+
+					UniqueIdSuffix++;
+				}
+			}
+
+			// End of CSS region
+//			Writer.LeaveRegion();
+		}
+
 		public override string ToString()
 		{
 			return Name;
@@ -183,34 +144,23 @@ namespace APIDocTool
 	class APIModuleIndex : APIPage
 	{
 		public APIModuleCategory Category;
-		public List<APIModule> Modules = new List<APIModule>();
-		public Dictionary<string, APIModule> ModuleLookup = new Dictionary<string,APIModule>();
+		public List<APIModule> Children = new List<APIModule>();
 
-		public APIModuleIndex(APIPage Parent, APIModuleCategory InCategory, IEnumerable<DoxygenModule> InModules)
+		public APIModuleIndex(APIPage Parent, APIModuleCategory InCategory)
 			: base(Parent, InCategory.Name)
 		{
 			Category = InCategory;
-
-			List<string> ModuleNames = new List<string>();
-			Category.GetModules(ModuleNames);
-
-			foreach (DoxygenModule InModule in InModules.Where(x => ModuleNames.Contains(x.Name)))
-			{
-				APIModule Module = APIModule.Build(this, InModule);
-				Modules.Add(Module);
-				ModuleLookup.Add(Module.Name, Module);
-			}
 		}
 
 		public override void GatherReferencedPages(List<APIPage> Pages)
 		{
-			Pages.AddRange(Modules);
+			Pages.AddRange(Children);
 		}
 
 		public override SitemapNode CreateSitemapNode()
 		{
 			SitemapNode Node = new SitemapNode(Name, SitemapLinkPath);
-			Node.Children.AddRange(Modules.OrderBy(x => x.Name).Select(x => x.CreateSitemapNode()));
+			Node.Children.AddRange(Children.OrderBy(x => x.Name).Select(x => x.CreateSitemapNode()));
 			return Node;
 		}
 
@@ -218,15 +168,20 @@ namespace APIDocTool
 		{
 			Manifest.Add("ModuleIndex:" + Name, this);
 
-			foreach (APIModule Module in Modules)
+			foreach (APIModule Child in Children)
 			{
-				Module.AddToManifest(Manifest);
+				Child.AddToManifest(Manifest);
 			}
 		}
 
 		public void WriteModuleList(UdnWriter Writer, int Depth)
 		{
-			Category.Write(Writer, Depth, ModuleLookup);
+			Category.Write(Writer, Depth);
+		}
+
+		public void WriteExpandableModuleList(UdnWriter Writer, string UniqueId)
+		{
+			Category.WriteExpandable(Writer, UniqueId);
 		}
 
 		public override void WritePage(UdnManifest Manifest, string OutputPath)
@@ -234,7 +189,9 @@ namespace APIDocTool
 			using (UdnWriter Writer = new UdnWriter(OutputPath))
 			{
 				Writer.WritePageHeader(Name, PageCrumbs, "Module Index");
+				Writer.EnterRegion("modules-list");
 				WriteModuleList(Writer, 2);
+				Writer.LeaveRegion();
 			}
 		}
 	}
diff --git a/Engine/Source/Programs/UnrealDocTool/APIDocTool/APIDocTool/APIRecord.cs b/Engine/Source/Programs/UnrealDocTool/APIDocTool/APIDocTool/APIRecord.cs
--- a/Engine/Source/Programs/UnrealDocTool/APIDocTool/APIDocTool/APIRecord.cs
+++ b/Engine/Source/Programs/UnrealDocTool/APIDocTool/APIDocTool/APIRecord.cs
@@ -43,7 +43,6 @@ namespace APIDocTool
 		public string TemplateSignature;
 		public string Definition;
 		public List<string> BaseDefinitions = new List<string>();
-		public string DeclarationFile;
 
 		public MetadataDirective MetadataDirective;
 
@@ -178,13 +177,6 @@ namespace APIDocTool
 				BaseDefinitions.Add(BaseDefinition);
 			}
 
-			// Parse the file that it's declared in
-			if(Entity.File != null)
-			{
-				int SourceIndex = Entity.File.IndexOf("/Engine/Source/");
-				if (SourceIndex != -1) DeclarationFile = Entity.File.Substring(SourceIndex);
-			}
-
 			// If it's a specialization, add it to the parent class
 			if(bIsTemplateSpecialization)
 			{
@@ -467,14 +459,9 @@ namespace APIDocTool
 				}
 				Writer.LeaveTag("[/PARAM]");
 
-				// Write the declaration file
-				Writer.EnterTag("[PARAM:declaration]");
-				if (!Utility.IsNullOrWhitespace(DeclarationFile))
-				{
-					Writer.EnterSection("declaration", "Declaration");
-					Writer.WriteEscapedLine(DeclarationFile);
-					Writer.LeaveSection();
-				}
+				// Write the references
+				Writer.EnterTag("[PARAM:references]");
+				WriteReferencesSection(Writer, Entity);
 				Writer.LeaveTag("[/PARAM]");
 
 				Writer.LeaveTag("[/OBJECT]");
diff --git a/Engine/Source/Programs/UnrealDocTool/APIDocTool/APIDocTool/APITypeDef.cs b/Engine/Source/Programs/UnrealDocTool/APIDocTool/APIDocTool/APITypeDef.cs
--- a/Engine/Source/Programs/UnrealDocTool/APIDocTool/APIDocTool/APITypeDef.cs
+++ b/Engine/Source/Programs/UnrealDocTool/APIDocTool/APIDocTool/APITypeDef.cs
@@ -7,21 +7,24 @@ using System.Text;
 using System.Threading.Tasks;
 using System.Xml;
 using System.IO;
+using DoxygenLib;
 
 namespace APIDocTool
 {
     class APITypeDef : APIMember
     {
+		public readonly DoxygenEntity Entity;
 		public readonly XmlNode Node;
 		public readonly string Id;
 		public string Type;
 		public string BriefDescription;
 		public string FullDescription;
 
-		public APITypeDef(APIPage InParent, XmlNode InNode) 
-			: base(InParent, InNode.SelectSingleNode("name").InnerText)
+		public APITypeDef(APIPage InParent, DoxygenEntity InEntity) 
+			: base(InParent, InEntity.Node.SelectSingleNode("name").InnerText)
         {
-			Node = InNode;
+			Entity = InEntity;
+			Node = Entity.Node;
 			Id = Node.Attributes["id"].Value;
 			AddRefLink(Id, this);
 		}
@@ -77,6 +80,11 @@ namespace APIDocTool
 				}
 				Writer.LeaveTag("[/PARAM]");
 
+				// Write the references
+				Writer.EnterTag("[PARAM:references]");
+				WriteReferencesSection(Writer, Entity);
+				Writer.LeaveTag("[/PARAM]");
+
 				// Leave the object tag
 				Writer.LeaveTag("[/OBJECT]");
 			}
diff --git a/Engine/Source/Programs/UnrealDocTool/APIDocTool/APIDocTool/IniFile.cs b/Engine/Source/Programs/UnrealDocTool/APIDocTool/APIDocTool/IniFile.cs
--- a/Engine/Source/Programs/UnrealDocTool/APIDocTool/APIDocTool/IniFile.cs
+++ b/Engine/Source/Programs/UnrealDocTool/APIDocTool/APIDocTool/IniFile.cs
@@ -112,6 +112,10 @@ namespace APIDocTool
 							CurrentSection = new IniSection(Name);
 
 							File.Sections.Add(CurrentSection);
+							if(File.NameToSection.ContainsKey(Name.ToLowerInvariant()))
+							{
+								Console.WriteLine("{0} has already been added", Name);
+							}
 							File.NameToSection.Add(Name.ToLowerInvariant(), CurrentSection);
 						}
 						else
diff --git a/Engine/Source/Programs/UnrealDocTool/APIDocTool/APIDocTool/Program.cs b/Engine/Source/Programs/UnrealDocTool/APIDocTool/APIDocTool/Program.cs
--- a/Engine/Source/Programs/UnrealDocTool/APIDocTool/APIDocTool/Program.cs
+++ b/Engine/Source/Programs/UnrealDocTool/APIDocTool/APIDocTool/Program.cs
@@ -55,8 +55,12 @@ namespace APIDocTool
 			"HTML5",
 			"Linux",
 			"TextureXboxOneFormat",
+			"NoRedist",
+			"NotForLicensees",
 		};
 
+		public static HashSet<string> ExcludeSourceDirectoriesHash = new HashSet<string>(ExcludeSourceDirectories.Select(x => x.ToLowerInvariant()));
+
 		static string[] ExcludeSourceFiles = 
 		{
 			"*/CoreUObject/Classes/Object.h",
@@ -186,6 +190,9 @@ namespace APIDocTool
 
 		static void Main(string[] Arguments)
         {
+			string EngineRootDir = Path.GetFullPath(Path.Combine(Path.GetDirectoryName(Assembly.GetExecutingAssembly().Location), "..\\..\\..\\..\\..\\..\\..\\.."));
+			string IntermediateRootDir = Path.Combine(EngineRootDir, "Intermediate\\Documentation");
+
 			bool bValidArgs = true;
 
 			bool bCleanMetadata = false;
@@ -200,10 +207,10 @@ namespace APIDocTool
 			bool bBuildHtml = false;
 			bool bBuildChm = false;
 
-			string TargetInfoPath = null;
-			string EngineDir = null;
-			string MetadataDir = null;
-			string XmlDir = null;
+			string TargetInfoPath = Path.Combine(IntermediateRootDir, "build\\targetinfo.xml");
+			string EngineDir = EngineRootDir;
+			string MetadataDir = Path.Combine(IntermediateRootDir, "metadata");
+			string XmlDir = Path.Combine(IntermediateRootDir, "doxygen");
 			string StatsPath = null;
 
 			bool bVerbose = false;
@@ -226,7 +233,7 @@ namespace APIDocTool
 				}
 				else if (OptionName == "-rebuild")
 				{
-					bCleanMetadata = bCleanXml = bCleanUdn = bCleanHtml = bBuildMetadata = bBuildXml = bBuildUdn = bBuildHtml = true;
+					bCleanMetadata = bCleanXml = bCleanUdn = bCleanHtml = bCleanChm = bBuildMetadata = bBuildXml = bBuildUdn = bBuildHtml = bBuildChm = true;
 				}
 				else if (OptionName == "-cleanmeta")
 				{
@@ -960,7 +967,13 @@ namespace APIDocTool
 		public static void CleanChm(string ChmPath)
 		{
 			Console.WriteLine("Cleaning '{0}'", ChmPath);
-			Utility.SafeDeleteDirectoryContents(ChmPath, true);
+			if (Directory.Exists(ChmPath))
+			{
+				foreach (string SubDir in Directory.EnumerateDirectories(ChmPath))
+				{
+					Directory.Delete(SubDir, true);
+				}
+			}
 		}
 
 		public static bool BuildChm(string ChmCompilerPath, string BaseHtmlDir, string ChmDir)
diff --git a/Engine/Source/Programs/UnrealDocTool/APIDocTool/APIDocTool/UdnManifest.cs b/Engine/Source/Programs/UnrealDocTool/APIDocTool/APIDocTool/UdnManifest.cs
--- a/Engine/Source/Programs/UnrealDocTool/APIDocTool/APIDocTool/UdnManifest.cs
+++ b/Engine/Source/Programs/UnrealDocTool/APIDocTool/APIDocTool/UdnManifest.cs
@@ -52,6 +52,50 @@ namespace APIDocTool
 			return Result;
 		}
 
+		public string FormatString(string Text)
+		{
+			StringBuilder Output = new StringBuilder();
+			for (int Idx = 0; Idx < Text.Length; Idx++)
+			{
+				if (Text[Idx] == '{')
+				{
+					if (Text[Idx + 1] == '{')
+					{
+						Output.Append(Text[++Idx]);
+					}
+					else
+					{
+						int EndIdx = Text.IndexOf('}', Idx + 1);
+						string LinkText = Text.Substring(Idx + 1, EndIdx - Idx - 1);
+						string LinkName = LinkText;
+
+						int LinkNameIdx = LinkText.IndexOf('|');
+						if (LinkNameIdx != -1)
+						{
+							LinkName = LinkText.Substring(LinkNameIdx + 1);
+							LinkText = LinkText.Substring(0, LinkNameIdx);
+						}
+
+						APIPage LinkPage = Find(LinkName);
+						if (LinkPage != null)
+						{
+							Output.AppendFormat("[{0}]({1})", LinkText, LinkPage.LinkPath);
+						}
+						else
+						{
+							Output.Append(LinkText);
+						}
+						Idx = EndIdx;
+					}
+				}
+				else
+				{
+					Output.Append(Text[Idx]);
+				}
+			}
+			return Output.ToString();
+		}
+
 		public void PrintConflicts()
 		{
 			foreach(KeyValuePair<string, List<APIPage>> ConflictEntry in ConflictEntries)
diff --git a/Engine/Source/Programs/UnrealDocTool/APIDocTool/APIDocTool/UdnWriter.cs b/Engine/Source/Programs/UnrealDocTool/APIDocTool/APIDocTool/UdnWriter.cs
--- a/Engine/Source/Programs/UnrealDocTool/APIDocTool/APIDocTool/UdnWriter.cs
+++ b/Engine/Source/Programs/UnrealDocTool/APIDocTool/APIDocTool/UdnWriter.cs
@@ -120,46 +120,7 @@ namespace APIDocTool
 
 		public void Write(UdnManifest Manifest, string Text)
 		{
-			StringBuilder Output = new StringBuilder();
-			for (int Idx = 0; Idx < Text.Length; Idx++)
-			{
-				if (Text[Idx] == '{')
-				{
-					if (Text[Idx + 1] == '{')
-					{
-						Output.Append(Text[++Idx]);
-					}
-					else
-					{
-						int EndIdx = Text.IndexOf('}', Idx + 1);
-						string LinkText = Text.Substring(Idx + 1, EndIdx - Idx - 1);
-						string LinkName = LinkText;
-
-						int LinkNameIdx = LinkText.IndexOf('|');
-						if (LinkNameIdx != -1)
-						{
-							LinkName = LinkText.Substring(LinkNameIdx + 1);
-							LinkText = LinkText.Substring(0, LinkNameIdx);
-						}
-
-						APIPage LinkPage = Manifest.Find(LinkName);
-						if (LinkPage != null)
-						{
-							Output.AppendFormat("[{0}]({1})", LinkText, LinkPage.LinkPath);
-						}
-						else
-						{
-							Output.Append(LinkText);
-						}
-						Idx = EndIdx;
-					}
-				}
-				else
-				{
-					Output.Append(Text[Idx]);
-				}
-			}
-			Write(Output.ToString());
+			Write(Manifest.FormatString(Text));
 		}
 
 		public void WriteLine()
@@ -451,6 +412,28 @@ namespace APIDocTool
 			LeaveTag("[/PARAMLITERAL]");
 		}
 
+		public void WriteList(IEnumerable<UdnListItem> Items)
+		{
+			// Write the head
+			WriteObject("MemberListHeadBlank");
+
+			// Write the body
+			foreach (UdnListItem Item in Items)
+			{
+				if (Item.Link == null)
+				{
+					WriteObject("MemberListItem", "name", Item.Name, "desc", Item.Description);
+				}
+				else
+				{
+					WriteObject("MemberListItemLinked", "name", Item.Name, "link", "[RELATIVE:" + Item.Link + "]", "desc", Item.Description);
+				}
+			}
+
+			// Write the tail
+			WriteObject("MemberListTail");
+		}
+
 		public void WriteList(string NameColumn, string DescColumn, IEnumerable<UdnListItem> Items)
 		{
 			// Write the head
diff --git a/Engine/Source/Programs/UnrealDocTool/DoxygenLib/DoxygenModule.cs b/Engine/Source/Programs/UnrealDocTool/DoxygenLib/DoxygenModule.cs
--- a/Engine/Source/Programs/UnrealDocTool/DoxygenLib/DoxygenModule.cs
+++ b/Engine/Source/Programs/UnrealDocTool/DoxygenLib/DoxygenModule.cs
@@ -42,12 +42,17 @@ namespace DoxygenLib
 		public string File;
 		public int Line;
 		public int Column;
+		public string BodyFile;
+		public int BodyStart;
+		public int BodyEnd;
+
 		public List<DoxygenEntity> Members = new List<DoxygenEntity>();
 
 		public DoxygenEntity(DoxygenModule InModule)
 		{
 			Module = InModule;
 			Line = Column = -1;
+			BodyStart = BodyEnd = -1;
 		}
 
 		public static DoxygenEntity FromXml(DoxygenModule InModule, string Name, XmlNode Node)
@@ -69,6 +74,16 @@ namespace DoxygenLib
 					Entity.Line = int.Parse(LineAttribute.Value);
 					Entity.Column = int.Parse(LocationNode.Attributes["column"].Value);
 				}
+
+				XmlAttribute BodyFileAttribute = LocationNode.Attributes["bodyfile"];
+				XmlAttribute BodyStartAttribute = LocationNode.Attributes["bodystart"];
+				XmlAttribute BodyEndAttribute = LocationNode.Attributes["bodyend"];
+				if (BodyFileAttribute != null && BodyStartAttribute != null && BodyEndAttribute != null)
+				{
+					Entity.BodyFile = BodyFileAttribute.Value.Replace('/', '\\');
+					Entity.BodyStart = Int32.Parse(BodyStartAttribute.Value);
+					Entity.BodyEnd = Int32.Parse(BodyEndAttribute.Value);
+				}
 			}
 			return Entity;
 		}
@@ -89,12 +104,54 @@ namespace DoxygenLib
 		}
 	}
 
+	public class DoxygenSourceFile
+	{
+		public string FileName;
+		public string NormalizedFileName;
+		public List<XmlNode> Lines = new List<XmlNode>();
+
+		public DoxygenSourceFile(string InFileName)
+		{
+			FileName = InFileName;
+			NormalizedFileName = NormalizeFileName(InFileName);
+		}
+
+		public static string NormalizeFileName(string InFileName)
+		{
+			return InFileName.Replace('\\', '/').ToLowerInvariant();
+		}
+
+		public static DoxygenSourceFile FromXml(XmlNode Node)
+		{
+			XmlNode LocationNode = Node.SelectSingleNode("location");
+			if(LocationNode == null) return null;
+
+			DoxygenSourceFile SourceFile = new DoxygenSourceFile(LocationNode.Attributes["file"].InnerText);
+			foreach(XmlNode CodeLineNode in Node.SelectNodes("programlisting/codeline"))
+			{
+				int CodeLineIdx = int.Parse(CodeLineNode.Attributes["lineno"].InnerText) - 1;
+				while(CodeLineIdx >= SourceFile.Lines.Count)
+				{
+					SourceFile.Lines.Add(null);
+				}
+				SourceFile.Lines[CodeLineIdx] = CodeLineNode;
+			}
+			return SourceFile;
+		}
+
+		public override string ToString()
+		{
+			return FileName;
+		}
+	}
+
 	public class DoxygenModule
 	{
 		public readonly string Name;
 		public readonly string BaseSrcDir;
 		public List<DoxygenCompound> Compounds = new List<DoxygenCompound>();
 		public List<DoxygenEntity> Entities = new List<DoxygenEntity>();
+		public List<DoxygenSourceFile> SourceFiles = new List<DoxygenSourceFile>();
 
 		public DoxygenModule(string InName, string InBaseSrcDir)
 		{
@@ -152,6 +209,8 @@ namespace DoxygenLib
 				if (Compound.Kind == "file")
 				{
 					ReadMembers(Module, Compound, "", null, Module.Entities);
+					DoxygenSourceFile SourceFile = DoxygenSourceFile.FromXml(Compound.Node);
+					if(SourceFile != null) Module.SourceFiles.Add(SourceFile);
 				}
 				else if (Compound.Kind == "namespace")
 				{
@@ -205,6 +264,23 @@ namespace DoxygenLib
 			}
 		}
 
+		protected static DoxygenSourceFile ReadSourceFile(string BaseXmlDir, string CompoundId)
+		{
+			string XmlFileName = Path.Combine(BaseXmlDir, CompoundId + ".xml");
+
+			XmlDocument Document;
+			if(TryReadXmlDocument(XmlFileName, out Document))
+			{
+				DoxygenSourceFile SourceFile = DoxygenSourceFile.FromXml(Document.SelectSingleNode("doxygen/compounddef"));
+				return SourceFile;
+			}
+			else
+			{
+				Console.WriteLine("Couldn't read source document: '{0}'", XmlFileName);
+				return null;
+			}
+		}
+
 		protected static void ReadMembers(DoxygenModule Module, DoxygenCompound Compound, string NamePrefix, DoxygenEntity Parent, List<DoxygenEntity> Entities)
 		{
 			using (XmlNodeList NodeList = Compound.Node.SelectNodes("sectiondef/memberdef"))
@@ -220,6 +296,19 @@ namespace DoxygenLib
 			}
 		}
 
+		public DoxygenSourceFile FindSourceFile(string InFileName)
+		{
+			string NormalizedFileName = DoxygenSourceFile.NormalizeFileName(InFileName);
+			foreach(DoxygenSourceFile SourceFile in SourceFiles)
+			{
+				if(SourceFile.NormalizedFileName == NormalizedFileName)
+				{
+					return SourceFile;
+				}
+			}
+			return null;
+		}
+
 		public override string ToString()
 		{
 			return Name;
diff --git a/Engine/Source/Programs/UnrealDocTool/DoxygenLib/Markdown.cs b/Engine/Source/Programs/UnrealDocTool/DoxygenLib/Markdown.cs
--- a/Engine/Source/Programs/UnrealDocTool/DoxygenLib/Markdown.cs
+++ b/Engine/Source/Programs/UnrealDocTool/DoxygenLib/Markdown.cs
@@ -137,10 +137,82 @@ namespace DoxygenLib
 			return output;
 		}
 
-		public static string EscapeText(string Text)
+		/** 
+		 * Parses Markdown from a Doxygen code node
+		 * 
+		 * @param Node			Xml node we want to parse into Markdown 
+		 * @param ResolveLink	Delegate to call to resolve doxygen id's into link paths. may be null.
+		 * @return				Converted text
+		 */
+		public static string ParseXmlCodeLine(XmlNode Node, ResolveLinkDelegate ResolveLink)
 		{
-			const string MarkdownCharacters = "_[]*{}()%#";
+			StringBuilder Builder = new StringBuilder();
+			ParseXmlCodeLine(Node, Builder, ResolveLink);
+			return Builder.ToString();
+		}
 
+		/** 
+		 * Parses Markdown from a Doxygen code node
+		 * 
+		 * @param Node			Xml node we want to parse into Markdown 
+		 * @param Output		StringBuilder to receive the output
+		 * @param ResolveLink	Delegate to call to resolve doxygen id's into link paths. may be null.
+		 */
+		static void ParseXmlCodeLine(XmlNode Node, StringBuilder Output, ResolveLinkDelegate ResolveLink)
+		{
+			switch(Node.Name)
+			{
+				case "sp":
+					Output.Append(' ');
+					return;
+				case "ref":
+					XmlAttribute RefAttribute = Node.Attributes["refid"];
+
+					string LinkPath = (RefAttribute != null && ResolveLink != null)? ResolveLink(RefAttribute.Value) : null;
+					if(LinkPath != null)
+					{
+						Output.Append("[");
+					}
+
+					ParseXmlCodeLineChildren(Node, Output, ResolveLink);
+
+					if(LinkPath != null)
+					{
+						Output.AppendFormat("]({0})", LinkPath);
+					}
+					break;
+				case "#text":
+					Output.Append(EscapeText(Node.InnerText));
+					break;
+				default:
+					ParseXmlCodeLineChildren(Node, Output, ResolveLink);
+					break;
+			}
+		}
+
+		/** 
+		 * Parses Markdown from the children of a Doxygen code node
+		 * 
+		 * @param Node			Xml node we want to parse into Markdown 
+		 * @param Output		StringBuilder to receive the output
+		 * @param ResolveLink	Delegate to call to resolve doxygen id's into link paths. may be null.
+		 */
+		static void ParseXmlCodeLineChildren(XmlNode Node, StringBuilder Output, ResolveLinkDelegate ResolveLink)
+		{
+			foreach(XmlNode ChildNode in Node.ChildNodes)
+			{
+				ParseXmlCodeLine(ChildNode, Output, ResolveLink);
+			}
+		}
+
+		/** 
+		 * Escape a string to be valid markdown.
+		 * 
+		 * @param Text			String to escape
+		 * @return				Escaped string
+		 */
+		public static string EscapeText(string Text)
+		{
 			StringBuilder Result = new StringBuilder();
 			for (int Idx = 0; Idx < Text.Length; Idx++)
 			{
@@ -164,13 +236,13 @@ namespace DoxygenLib
 				{
 					Result.Append("    ");
 				}
-				else if (Text[Idx] < 0x20 || Text[Idx] > 0x7e || MarkdownCharacters.Contains(Text[Idx]))
+				else if(Char.IsLetterOrDigit(Text[Idx]) || Text[Idx] == '_' || Text[Idx] == ' ' || Text[Idx] == ',' || Text[Idx] == '.')
 				{
-					Result.AppendFormat("&#{0:00};", (int)Text[Idx]);
+					Result.Append(Text[Idx]);
 				}
 				else
 				{
-					Result.Append(Text[Idx]);
+					Result.AppendFormat("&#{0:00};", (int)Text[Idx]);
 				}
 			}
 			return Result.ToString();
