Tweaks to generated API documentation.

- Move fully expanded module list onto front page.
- Separate enum and class index pages.
- Move index building onto a command line parameter.
- Allow running UBT from API doc tool.

diff --git a/Engine/Source/Programs/AutomationTool/Scripts/BuildDocumentation.Automation.cs b/Engine/Source/Programs/AutomationTool/Scripts/BuildDocumentation.Automation.cs
--- a/Engine/Source/Programs/AutomationTool/Scripts/BuildDocumentation.Automation.cs
+++ b/Engine/Source/Programs/AutomationTool/Scripts/BuildDocumentation.Automation.cs
@@ -12,12 +12,6 @@ using UnrealBuildTool;
 [Help("-env", "Builds the compile environment")]
 public class BuildDocumentation : BuildCommand
 {
-	const string EngineTarget = "DocumentationEditor";
-	const string EnginePlatform = "Win64";
-	const string EngineConfiguration = "Debug";
-
-	const string EngineUProject = "Engine\\Documentation\\Extras\\API\\Build\\Documentation.uproject";
-
 	public override void ExecuteBuild()
 	{
 		// Parse the command line
@@ -90,7 +84,7 @@ public class BuildDocumentation : BuildCommand
 		// Execute the clean
 		if(bCleanEnv)
 		{
-			File.Delete(TargetInfoPath);
+			RunAndLog(CmdEnv, ApiDocToolPath, "-cleantargetinfo" + ApiToolCommandLine, "ApiDocTool-CleanTargetInfo");
 		}
 		if (bCleanMeta)
 		{
@@ -116,7 +110,7 @@ public class BuildDocumentation : BuildCommand
 		// Execute the build
 		if (bBuildEnv)
 		{
-			CommandUtils.RunUBT(CmdEnv, UnrealBuildToolPath, String.Format("{0} {1} {2} -project=\"{3}\" -writetargetinfo=\"{4}\"", EngineTarget, EnginePlatform, EngineConfiguration, Path.GetFullPath(Path.Combine(CmdEnv.LocalRoot, EngineUProject)), TargetInfoPath));
+			RunAndLog(CmdEnv, ApiDocToolPath, "-buildtargetinfo" + ApiToolCommandLine, "ApiDocTool-BuildTargetInfo");
 		}
 		if (bBuildMeta)
 		{
diff --git a/Engine/Source/Programs/UnrealDocTool/APIDocTool/APIDocTool/APIDocTool.csproj b/Engine/Source/Programs/UnrealDocTool/APIDocTool/APIDocTool/APIDocTool.csproj
--- a/Engine/Source/Programs/UnrealDocTool/APIDocTool/APIDocTool/APIDocTool.csproj
+++ b/Engine/Source/Programs/UnrealDocTool/APIDocTool/APIDocTool/APIDocTool.csproj
@@ -67,6 +67,7 @@
     </Compile>
     <Compile Include="APIConstant.cs" />
     <Compile Include="APIConstantIndex.cs" />
+    <Compile Include="APIEnumIndex.cs" />
     <Compile Include="APIFunctionIndex.cs" />
     <Compile Include="APIHierarchy.cs" />
     <Compile Include="APIIndex.cs" />
@@ -76,7 +77,7 @@
     <Compile Include="APIEventParameters.cs" />
     <Compile Include="APIGettingStarted.cs" />
     <Compile Include="APIRecord.cs" />
-    <Compile Include="APITypeIndex.cs" />
+    <Compile Include="APIRecordIndex.cs" />
     <Compile Include="BuildTarget.cs" />
     <Compile Include="Icons.cs" />
     <Compile Include="IniFile.cs" />
diff --git a/Engine/Source/Programs/UnrealDocTool/APIDocTool/APIDocTool/APIEnumIndex.cs b/Engine/Source/Programs/UnrealDocTool/APIDocTool/APIDocTool/APIEnumIndex.cs
new file mode 100644
--- /dev/null
+++ b/Engine/Source/Programs/UnrealDocTool/APIDocTool/APIDocTool/APIEnumIndex.cs
@@ -0,0 +1,31 @@
+﻿// Copyright 1998-2014 Epic Games, Inc. All Rights Reserved.
+
+using System;
+using System.Collections.Generic;
+using System.Linq;
+using System.Text;
+using System.Threading.Tasks;
+
+namespace APIDocTool
+{
+	class APIEnumIndex : APIMemberIndex
+	{
+		public const string IgnorePrefixLetters = "E";
+
+		public APIEnumIndex(APIPage InParent, IEnumerable<APIMember> Members) : base(InParent, "Enums", "All enums")
+		{
+			foreach (APIMember Member in Members)
+			{
+				if (Member is APIEnum)
+				{
+					Entry NewEntry = new Entry(Member);
+					if (NewEntry.SortName.Length >= 2 && IgnorePrefixLetters.Contains(NewEntry.SortName[0]) && Char.IsUpper(NewEntry.SortName[1]))
+					{
+						NewEntry.SortName = NewEntry.SortName.Substring(1);
+					}
+					AddToDefaultCategory(NewEntry);
+				}
+			}
+		}
+	}
+}
diff --git a/Engine/Source/Programs/UnrealDocTool/APIDocTool/APIDocTool/APIIndex.cs b/Engine/Source/Programs/UnrealDocTool/APIDocTool/APIDocTool/APIIndex.cs
--- a/Engine/Source/Programs/UnrealDocTool/APIDocTool/APIDocTool/APIIndex.cs
+++ b/Engine/Source/Programs/UnrealDocTool/APIDocTool/APIDocTool/APIIndex.cs
@@ -18,10 +18,10 @@ namespace APIDocTool
 		public List<APIModuleIndex> ChildModuleIndexes = new List<APIModuleIndex>();
 
 		public APIGettingStarted GettingStarted;
-		public APIModuleIndex ModuleIndex;
 		public APIConstantIndex ConstantIndex;
 		public APIFunctionIndex FunctionIndex;
-		public APITypeIndex TypeIndex;
+		public APIEnumIndex EnumIndex;
+		public APIRecordIndex RecordIndex;
 		public APIHierarchy RecordHierarchy;
 
 		public static APIModuleCategory FindRootCategoryByPath(List<APIModuleCategory> Categories, string InPath)
@@ -37,7 +37,7 @@ namespace APIDocTool
 				int CategoryMaxIdx = CategoryMinIdx;
 				while (CategoryMaxIdx < NormalizedSrcPath.Length && NormalizedSrcPath[CategoryMaxIdx] != '/') CategoryMaxIdx++;
 
-				string CategoryName = NormalizedSrcPath.Substring(CategoryMinIdx, CategoryMaxIdx - CategoryMinIdx) + " Modules";
+				string CategoryName = NormalizedSrcPath.Substring(CategoryMinIdx, CategoryMaxIdx - CategoryMinIdx);
 
 				APIModuleCategory Category = Categories.FirstOrDefault(x => x.Name.ToLower() == CategoryName.ToLower());
 				if(Category != null) return Category;
@@ -58,9 +58,9 @@ namespace APIDocTool
 
 			// Add some default categories to fix the order
 			APIModuleCategory RootCategory = new APIModuleCategory("Modules");
-			RootCategory.Categories.Add(new APIModuleCategory("Runtime Modules"));
-			RootCategory.Categories.Add(new APIModuleCategory("Editor Modules"));
-			RootCategory.Categories.Add(new APIModuleCategory("Developer Modules"));
+			RootCategory.Categories.Add(new APIModuleCategory("Runtime"));
+			RootCategory.Categories.Add(new APIModuleCategory("Editor"));
+			RootCategory.Categories.Add(new APIModuleCategory("Developer"));
 
 			// Add all of the modules to a category
 			Dictionary<DoxygenModule, APIModuleCategory> ModuleToCategory = new Dictionary<DoxygenModule, APIModuleCategory>();
@@ -108,10 +108,6 @@ namespace APIDocTool
 			GettingStarted = new APIGettingStarted(this);
 			Pages.Add(GettingStarted);
 
-			// Create an index of all the modules
-			ModuleIndex = new APIModuleIndex(this, RootCategory);
-			Pages.Add(ModuleIndex);
-
 			// Create an index of all the constants
 			ConstantIndex = new APIConstantIndex(this, AllMembers);
 			Pages.Add(ConstantIndex);
@@ -120,9 +116,13 @@ namespace APIDocTool
 			FunctionIndex = new APIFunctionIndex(this, AllMembers.OfType<APIFunction>().Where(x => !(x.ScopeParent is APIRecord)));
 			Pages.Add(FunctionIndex);
 
-			// Create an index of all the types
-			TypeIndex = new APITypeIndex(this, AllMembers);
-			Pages.Add(TypeIndex);
+			// Create an index of all the enums
+			EnumIndex = new APIEnumIndex(this, AllMembers);
+			Pages.Add(EnumIndex);
+
+			// Create an index of all the records
+			RecordIndex = new APIRecordIndex(this, AllMembers);
+			Pages.Add(RecordIndex);
 
 			// Create an index of all the classes
 			RecordHierarchy = APIHierarchy.Build(this, AllMembers.OfType<APIRecord>());
@@ -138,10 +138,10 @@ namespace APIDocTool
 			{
 				RootNodes.Add(ChildModuleIndex.CreateSitemapNode());
 			}
-			RootNodes.Add(new SitemapNode("All modules", ModuleIndex.SitemapLinkPath));
 			RootNodes.Add(new SitemapNode("All constants", ConstantIndex.SitemapLinkPath));
 			RootNodes.Add(new SitemapNode("All functions", FunctionIndex.SitemapLinkPath));
-			RootNodes.Add(new SitemapNode("All types", TypeIndex.SitemapLinkPath));
+			RootNodes.Add(new SitemapNode("All enums", EnumIndex.SitemapLinkPath));
+			RootNodes.Add(new SitemapNode("All classes", RecordIndex.SitemapLinkPath));
 			RootNodes.Add(new SitemapNode("Class hierarchy", RecordHierarchy.SitemapLinkPath));
 			Sitemap.WriteContentsFile(OutputPath, RootNodes);
 		}
@@ -213,24 +213,20 @@ namespace APIDocTool
 
 				Writer.WriteLine("The API reference is an early work in progress, and some information may be missing or out of date. It serves mainly as a low level index of Engine classes and functions.");
 				Writer.WriteLine("For tutorials, walkthroughs and detailed guides to programming with Unreal, please see the [Unreal Engine Programming](http://docs.unrealengine.com/latest/INT/Programming/index.html) home on the web.");
-
-				Writer.WriteHeading(2, "Introduction");
-
-				Writer.WriteLine(Manifest, "The Unreal Engine consists of a large number of modules divided into three categories: the {Runtime|ModuleIndex:Runtime Modules}, functionality for the {Editor|ModuleIndex:Editor Modules}, and {Developer|ModuleIndex:Developer Modules} utilities. Full documentation for each of these categories is below.");
 				Writer.WriteLine();
 				Writer.WriteLine("To explore the API from some of the most frequently encountered Unreal concepts and types, see the API __[getting started](" + GettingStarted.LinkPath + ")__ page.");
 
 				Writer.WriteHeading(2, "Contents");
 
 				Writer.EnterRegion("memberindexlinks");
-				Writer.WriteLine("[Getting started]({0}) &middot; [All modules]({1}) &middot; [All constants]({2}) &middot; [All functions]({3}) &middot; [All types]({4}) &middot; [Class hierarchy]({5})", GettingStarted.LinkPath, ModuleIndex.LinkPath, ConstantIndex.LinkPath, FunctionIndex.LinkPath, TypeIndex.LinkPath, RecordHierarchy.LinkPath);
+				Writer.WriteLine("[Getting started]({0}) &middot; [All constants]({1}) &middot; [All functions]({2}) &middot; [All enums]({3}) &middot; [All classes]({4}) &middot; [Class hierarchy]({5})", GettingStarted.LinkPath, ConstantIndex.LinkPath, FunctionIndex.LinkPath, EnumIndex.LinkPath, RecordIndex.LinkPath, RecordHierarchy.LinkPath);
 				Writer.LeaveRegion();
 
 				foreach(APIModuleIndex ChildModuleIndex in ChildModuleIndexes)
 				{
 					Writer.WriteHeading(2, ChildModuleIndex.Name);
 					Writer.EnterRegion("modules-list");
-					ChildModuleIndex.WriteExpandableModuleList(Writer, ChildModuleIndex.Name.Replace(' ', '_'));
+					ChildModuleIndex.WriteModuleList(Writer, 3);
 					Writer.LeaveRegion();
 				}
 
diff --git a/Engine/Source/Programs/UnrealDocTool/APIDocTool/APIDocTool/APIRecordIndex.cs b/Engine/Source/Programs/UnrealDocTool/APIDocTool/APIDocTool/APIRecordIndex.cs
new file mode 100644
--- /dev/null
+++ b/Engine/Source/Programs/UnrealDocTool/APIDocTool/APIDocTool/APIRecordIndex.cs
@@ -0,0 +1,31 @@
+﻿// Copyright 1998-2014 Epic Games, Inc. All Rights Reserved.
+
+using System;
+using System.Collections.Generic;
+using System.Linq;
+using System.Text;
+using System.Threading.Tasks;
+
+namespace APIDocTool
+{
+	class APIRecordIndex : APIMemberIndex
+	{
+		public const string IgnorePrefixLetters = "UAFTS"; 
+
+		public APIRecordIndex(APIPage InParent, IEnumerable<APIMember> Members) : base(InParent, "Classes", "All classes")
+		{
+			foreach (APIMember Member in Members)
+			{
+				if ((Member is APIRecord && (!((APIRecord)Member).bIsTemplateSpecialization)) || (Member is APITypeDef) || (Member is APIEventParameters))
+				{
+					Entry NewEntry = new Entry(Member);
+					if (NewEntry.SortName.Length >= 2 && IgnorePrefixLetters.Contains(NewEntry.SortName[0]) && Char.IsUpper(NewEntry.SortName[1]))
+					{
+						NewEntry.SortName = NewEntry.SortName.Substring(1);
+					}
+					AddToDefaultCategory(NewEntry);
+				}
+			}
+		}
+	}
+}
diff --git a/Engine/Source/Programs/UnrealDocTool/APIDocTool/APIDocTool/APITypeIndex.cs b/Engine/Source/Programs/UnrealDocTool/APIDocTool/APIDocTool/APITypeIndex.cs
deleted file mode 100644
diff --git a/Engine/Source/Programs/UnrealDocTool/APIDocTool/APIDocTool/Program.cs b/Engine/Source/Programs/UnrealDocTool/APIDocTool/APIDocTool/Program.cs
--- a/Engine/Source/Programs/UnrealDocTool/APIDocTool/APIDocTool/Program.cs
+++ b/Engine/Source/Programs/UnrealDocTool/APIDocTool/APIDocTool/Program.cs
@@ -1,7 +1,5 @@
 ﻿// Copyright 1998-2014 Epic Games, Inc. All Rights Reserved.
 
-//#define TEST_INDEX_PAGE
-
 using System;
 using System.Collections.Generic;
 using System.Linq;
@@ -37,6 +35,7 @@ namespace APIDocTool
 		public const string TabSpaces = "&nbsp;&nbsp;&nbsp;&nbsp;";
         public const string APIFolder = "API";
 
+		public static bool bIndexOnly = false;
 		public static bool bOutputPublish = false;
 
 		static string[] ExcludeSourceDirectories =
@@ -160,6 +159,12 @@ namespace APIDocTool
 		const string SitemapContentsFileName = "API.hhc";
 		const string SitemapIndexFileName = "API.hhk";
 
+		const string TargetInfoName = "DocumentationEditor";
+		const string TargetInfoPlatform = "Win64";
+		const string TargetInfoConfig = "Debug";
+
+		const string TargetInfoProject = "Documentation\\Extras\\API\\Build\\Documentation.uproject";
+
 		static Program()
 		{
 			List<string> DelegateMacros = new List<string>();
@@ -195,12 +200,14 @@ namespace APIDocTool
 
 			bool bValidArgs = true;
 
+			bool bCleanTargetInfo = false;
 			bool bCleanMetadata = false;
 			bool bCleanXml = false;
 			bool bCleanUdn = false;
 			bool bCleanHtml = false;
 			bool bCleanChm = false;
 
+			bool bBuildTargetInfo = false;
 			bool bBuildMetadata = false;
 			bool bBuildXml = false;
 			bool bBuildUdn = false;
@@ -225,15 +232,24 @@ namespace APIDocTool
 
 				if (OptionName == "-clean")
 				{
-					bCleanMetadata = bCleanXml = bCleanUdn = bCleanHtml = true;
+					bCleanTargetInfo = bCleanMetadata = bCleanXml = bCleanUdn = bCleanHtml = true;
 				}
 				else if (OptionName == "-build")
 				{
-					bBuildMetadata = bBuildXml = bBuildUdn = bBuildHtml = true;
+					bBuildTargetInfo = bBuildMetadata = bBuildXml = bBuildUdn = bBuildHtml = true;
 				}
 				else if (OptionName == "-rebuild")
 				{
-					bCleanMetadata = bCleanXml = bCleanUdn = bCleanHtml = bCleanChm = bBuildMetadata = bBuildXml = bBuildUdn = bBuildHtml = bBuildChm = true;
+					bCleanTargetInfo = bCleanMetadata = bCleanXml = bCleanUdn = bCleanHtml = bCleanChm = true;
+					bBuildTargetInfo = bBuildMetadata = bBuildXml = bBuildUdn = bBuildHtml = bBuildChm = true;
+				}
+				else if (OptionName == "-cleantarget")
+				{
+					bCleanTargetInfo = true;
+				}
+				else if (OptionName == "-buildtarget")
+				{
+					bBuildTargetInfo = true;
 				}
 				else if (OptionName == "-cleanmeta")
 				{
@@ -319,6 +335,10 @@ namespace APIDocTool
 				{
 					bVerbose = true;
 				}
+				else if (OptionName == "-indexonly")
+				{
+					bIndexOnly = true;
+				}
 				else if (OptionName == "-filter=")
 				{
 					Filters.AddRange(OptionValue.Split(',').Select(x => x.Replace('\\', '/').Trim()));
@@ -351,17 +371,16 @@ namespace APIDocTool
 				string SettingsPath = Path.Combine(EngineDir, "Documentation\\Extras\\API\\API.ini");
 				string MetadataPath = Path.Combine(MetadataDir, "metadata.xml");
 
-				// Create all the output directories
-				Utility.SafeCreateDirectory(UdnDir);
-				Utility.SafeCreateDirectory(HtmlDir);
-				Utility.SafeCreateDirectory(ChmDir);
-
 				// Read the settings file
 				Settings = IniFile.Read(SettingsPath);
 				IgnoredFunctionMacros = new HashSet<string>(Settings.FindValueOrDefault("Input.IgnoredFunctionMacros", "").Split('\n'));
 				IncludedSourceFiles = new HashSet<string>(Settings.FindValueOrDefault("Output.IncludedSourceFiles", "").Split('\n'));
 
 				// Clean the output folders
+				if (bCleanTargetInfo)
+				{
+					CleanTargetInfo(TargetInfoPath);
+				}
 				if (bCleanMetadata)
 				{
 					CleanMetadata(MetadataDir);
@@ -384,17 +403,20 @@ namespace APIDocTool
 				}
 
 				// Build the data
-				if (!bBuildMetadata || BuildMetadata(DoxygenPath, EngineDir, MetadataDir, MetadataPath))
+				if (!bBuildTargetInfo || BuildTargetInfo(TargetInfoPath, EngineDir))
 				{
-					if (!bBuildXml || BuildXml(EngineDir, TargetInfoPath, DoxygenPath, XmlDir, Filters, bVerbose))
+					if (!bBuildMetadata || BuildMetadata(DoxygenPath, EngineDir, MetadataDir, MetadataPath))
 					{
-						if (!bBuildUdn || BuildUdn(XmlDir, UdnDir, ChmDir, MetadataPath, StatsPath, Filters))
+						if (!bBuildXml || BuildXml(EngineDir, TargetInfoPath, DoxygenPath, XmlDir, Filters, bVerbose))
 						{
-							if (!bBuildHtml || BuildHtml(EngineDir, DocToolPath, UdnDir))
+							if (!bBuildUdn || BuildUdn(XmlDir, UdnDir, ChmDir, MetadataPath, StatsPath, Filters))
 							{
-								if (!bBuildChm || BuildChm(ChmCompilerPath, HtmlDir, ChmDir))
+								if (!bBuildHtml || BuildHtml(EngineDir, DocToolPath, UdnDir, HtmlDir))
 								{
-									Console.WriteLine("Complete.");
+									if (!bBuildChm || BuildChm(ChmCompilerPath, HtmlDir, ChmDir))
+									{
+										Console.WriteLine("Complete.");
+									}
 								}
 							}
 						}
@@ -423,6 +445,46 @@ namespace APIDocTool
 			}
 		}
 
+		static void CleanTargetInfo(string TargetInfoPath)
+		{
+			Console.WriteLine("Cleaning '{0}'", TargetInfoPath);
+			Utility.SafeDeleteFile(TargetInfoPath);
+		}
+
+		static bool BuildTargetInfo(string TargetInfoPath, string EngineDir)
+		{
+			Console.WriteLine("Building target info...");
+
+			Utility.SafeCreateDirectory(Path.GetDirectoryName(TargetInfoPath));
+
+			using (Process BuildToolProcess = new Process())
+			{
+				BuildToolProcess.StartInfo.WorkingDirectory = EngineDir;
+				BuildToolProcess.StartInfo.FileName = Path.Combine(EngineDir, "Binaries\\DotNET\\UnrealBuildTool.exe");
+				BuildToolProcess.StartInfo.Arguments = String.Format("{0} {1} {2} -project=\"{3}\" -writetargetinfo=\"{4}\"", TargetInfoName, TargetInfoPlatform, TargetInfoConfig, Path.Combine(EngineDir, TargetInfoProject), TargetInfoPath);
+				BuildToolProcess.StartInfo.UseShellExecute = false;
+				BuildToolProcess.StartInfo.RedirectStandardOutput = true;
+				BuildToolProcess.StartInfo.RedirectStandardError = true;
+
+				BuildToolProcess.OutputDataReceived += new DataReceivedEventHandler(ProcessOutputReceived);
+				BuildToolProcess.ErrorDataReceived += new DataReceivedEventHandler(ProcessOutputReceived);
+
+				try
+				{
+					BuildToolProcess.Start();
+					BuildToolProcess.BeginOutputReadLine();
+					BuildToolProcess.BeginErrorReadLine();
+					BuildToolProcess.WaitForExit();
+					return BuildToolProcess.ExitCode == 0;
+				}
+				catch (Exception Ex)
+				{
+					Console.WriteLine(Ex.ToString() + "\n" + Ex.StackTrace);
+					return false;
+				}
+			}
+		}
+
 		static void CleanMetadata(string MetadataDir)
 		{
 			Console.WriteLine("Cleaning '{0}'", MetadataDir);
@@ -494,6 +556,9 @@ namespace APIDocTool
 
 		static bool BuildXml(string EngineDir, string TargetInfoPath, string DoxygenPath, string XmlDir, List<string> Filters = null, bool bVerbose = false)
 		{
+			// Create the output directory
+			Utility.SafeCreateDirectory(XmlDir);
+
 			// Read the target that we're building
 			BuildTarget Target = BuildTarget.Read(TargetInfoPath);
 
@@ -516,79 +581,82 @@ namespace APIDocTool
 				}
 			}
 
-#if TEST_INDEX_PAGE
 			// Just use all the input modules
-			List<BuildModule> ParsedModules = new List<BuildModule>(InputModules);
-#else
-			// Filter the input module list
-			if (Filters != null && Filters.Count > 0)
+			List<BuildModule> ParsedModules = new List<BuildModule>();
+			if(bIndexOnly)
 			{
-				InputModules = new List<BuildModule>(InputModules.Where(x => Filters.Exists(y => y.StartsWith(x.Name + "/", StringComparison.InvariantCultureIgnoreCase))));
+				ParsedModules.AddRange(InputModules);
 			}
-
-			// Sort all the modules into the required build order
-			Console.WriteLine("Sorting module parse order...");
-			List<BuildModule> SortedModules = SortModules(InputModules);
-
-			// List all the modules which have circular dependencies
-			foreach (BuildModule Module in SortedModules)
+			else
 			{
-				if (SortedModules.Count(x => x == Module) > 1)
+				// Filter the input module list
+				if (Filters != null && Filters.Count > 0)
 				{
-					Console.WriteLine("  Circular dependency on {0}; module will be parsed twice", Module.Name);
+					InputModules = new List<BuildModule>(InputModules.Where(x => Filters.Exists(y => y.StartsWith(x.Name + "/", StringComparison.InvariantCultureIgnoreCase))));
 				}
-			}
 
-			// Set our error mode so as to not bring up the WER dialog if Doxygen crashes (which it often does)
-			SetErrorMode(0x0007);
+				// Sort all the modules into the required build order
+				Console.WriteLine("Sorting module parse order...");
+				List<BuildModule> SortedModules = SortModules(InputModules);
 
-			// Create the output directory
-			Utility.SafeCreateDirectory(XmlDir);
-
-			// Build all the modules
-			List<BuildModule> ParsedModules = new List<BuildModule>();
-			Dictionary<string, string> ModuleToTagfile = new Dictionary<string, string>();
-			for (int Idx = 0; Idx < SortedModules.Count; Idx++)
-			{
-				BuildModule Module = SortedModules[Idx];
-				Console.WriteLine("Parsing source for {0}... ({1}/{2})", Module.Name, Idx + 1, SortedModules.Count);
-				string ModuleDir = Path.Combine(XmlDir, Module.Name);
-
-				// Build the configuration for this module
-				DoxygenConfig Config = new DoxygenConfig(Module.Name, Module.Path, ModuleDir);
-				Config.Definitions.AddRange(Module.Definitions.Select(x => x.Replace("DLLIMPORT", "").Replace("DLLEXPORT", "")));
-				Config.Definitions.AddRange(DoxygenPredefinedMacros);
-				Config.ExpandAsDefined.AddRange(DoxygenExpandedMacros);
-				Config.IncludePaths.AddRange(Module.IncludePaths.Select(x => ExpandIncludePath(EngineDir, x).TrimEnd('\\')));
-				Config.ExcludePatterns.AddRange(ExcludeSourceDirectories.Select(x => "*/" + x + "/*"));
-				Config.ExcludePatterns.AddRange(ExcludeSourceFiles);
-				Config.OutputTagfile = Path.Combine(ModuleDir, "tags.xml");
-
-				// Add all the valid tagfiles
-				foreach (string Dependency in Module.Dependencies)
-				{
-					string Tagfile;
-					if (ModuleToTagfile.TryGetValue(Dependency, out Tagfile))
+				// List all the modules which have circular dependencies
+				foreach (BuildModule Module in SortedModules)
+				{
+					if (SortedModules.Count(x => x == Module) > 1)
 					{
-						Config.InputTagfiles.Add(Tagfile);
+						Console.WriteLine("  Circular dependency on {0}; module will be parsed twice", Module.Name);
 					}
 				}
 
-				// Run doxygen
-				if (Doxygen.Run(DoxygenPath, Path.Combine(EngineDir, "Source"), Config, bVerbose))
+				// Set our error mode so as to not bring up the WER dialog if Doxygen crashes (which it often does)
+				SetErrorMode(0x0007);
+
+				// Create the output directory
+				Utility.SafeCreateDirectory(XmlDir);
+
+				// Build all the modules
+				Dictionary<string, string> ModuleToTagfile = new Dictionary<string, string>();
+				for (int Idx = 0; Idx < SortedModules.Count; Idx++)
 				{
-					if (!ModuleToTagfile.ContainsKey(Module.Name))
+					BuildModule Module = SortedModules[Idx];
+					Console.WriteLine("Parsing source for {0}... ({1}/{2})", Module.Name, Idx + 1, SortedModules.Count);
+					string ModuleDir = Path.Combine(XmlDir, Module.Name);
+
+					// Build the configuration for this module
+					DoxygenConfig Config = new DoxygenConfig(Module.Name, Module.Path, ModuleDir);
+					Config.Definitions.AddRange(Module.Definitions.Select(x => x.Replace("DLLIMPORT", "").Replace("DLLEXPORT", "")));
+					Config.Definitions.AddRange(DoxygenPredefinedMacros);
+					Config.ExpandAsDefined.AddRange(DoxygenExpandedMacros);
+					Config.IncludePaths.AddRange(Module.IncludePaths.Select(x => ExpandIncludePath(EngineDir, x).TrimEnd('\\')));
+					Config.ExcludePatterns.AddRange(ExcludeSourceDirectories.Select(x => "*/" + x + "/*"));
+					Config.ExcludePatterns.AddRange(ExcludeSourceFiles);
+					Config.OutputTagfile = Path.Combine(ModuleDir, "tags.xml");
+
+					// Add all the valid tagfiles
+					foreach (string Dependency in Module.Dependencies)
+					{
+						string Tagfile;
+						if (ModuleToTagfile.TryGetValue(Dependency, out Tagfile))
+						{
+							Config.InputTagfiles.Add(Tagfile);
+						}
+					}
+
+					// Run doxygen
+					if (Doxygen.Run(DoxygenPath, Path.Combine(EngineDir, "Source"), Config, bVerbose))
 					{
-						ModuleToTagfile.Add(Module.Name, Config.OutputTagfile);
-						ParsedModules.Add(Module);
+						if (!ModuleToTagfile.ContainsKey(Module.Name))
+						{
+							ModuleToTagfile.Add(Module.Name, Config.OutputTagfile);
+							ParsedModules.Add(Module);
+						}
+					}
+					else
+					{
+						Console.WriteLine("  Doxygen crashed. Skipping.");
 					}
-				}
-				else
-				{
-					Console.WriteLine("  Doxygen crashed. Skipping.");
 				}
 			}
-#endif
 
 			// Write the modules file
 			WriteModulesXml(XmlDir, ParsedModules);
@@ -672,6 +740,9 @@ namespace APIDocTool
 
 		static bool BuildUdn(string XmlDir, string UdnDir, string SitemapDir, string MetadataPath, string StatsPath, List<string> Filters = null)
 		{
+			// Create the output directory
+			Utility.SafeCreateDirectory(UdnDir);
+
 			// Read the metadata
 			MetadataLookup.Load(MetadataPath);
 
@@ -690,34 +761,36 @@ namespace APIDocTool
 				}
 			}
 
-#if TEST_INDEX_PAGE
-			// Just create empty modules
 			List<DoxygenModule> Modules = new List<DoxygenModule>();
-			for (int Idx = 0; Idx < InputModules.Count; Idx++)
+			if(bIndexOnly)
 			{
-				Modules.Add(new DoxygenModule(InputModules[Idx].Key, InputModules[Idx].Value));
+				// Just create empty modules
+				for (int Idx = 0; Idx < InputModules.Count; Idx++)
+				{
+					Modules.Add(new DoxygenModule(InputModules[Idx].Key, InputModules[Idx].Value));
+				}
 			}
-#else
-			// Filter the input module list
-			if (Filters != null && Filters.Count > 0)
+			else
 			{
-				InputModules = new List<KeyValuePair<string, string>>(InputModules.Where(x => Filters.Exists(y => y.StartsWith(x.Key + "/", StringComparison.InvariantCultureIgnoreCase))));
-			}
+				// Filter the input module list
+				if (Filters != null && Filters.Count > 0)
+				{
+					InputModules = new List<KeyValuePair<string, string>>(InputModules.Where(x => Filters.Exists(y => y.StartsWith(x.Key + "/", StringComparison.InvariantCultureIgnoreCase))));
+				}
 
-			// Read all the doxygen modules
-			List<DoxygenModule> Modules = new List<DoxygenModule>();
-			for(int Idx = 0; Idx < InputModules.Count; Idx++)
-			{
-				Console.WriteLine("Reading module {0}... ({1}/{2})", InputModules[Idx].Key, Idx + 1, InputModules.Count);
-				Modules.Add(DoxygenModule.Read(InputModules[Idx].Key, InputModules[Idx].Value, Path.Combine(XmlDir, InputModules[Idx].Key, "xml")));
-			}
+				// Read all the doxygen modules
+				for(int Idx = 0; Idx < InputModules.Count; Idx++)
+				{
+					Console.WriteLine("Reading module {0}... ({1}/{2})", InputModules[Idx].Key, Idx + 1, InputModules.Count);
+					Modules.Add(DoxygenModule.Read(InputModules[Idx].Key, InputModules[Idx].Value, Path.Combine(XmlDir, InputModules[Idx].Key, "xml")));
+				}
 
-			// Now filter all the entities in each module
-			if(Filters != null && Filters.Count > 0)
-			{
-				FilterEntities(Modules, Filters);
+				// Now filter all the entities in each module
+				if(Filters != null && Filters.Count > 0)
+				{
+					FilterEntities(Modules, Filters);
+				}
 			}
-#endif
 
 			// Create the index page, and all the pages below it
 			APIIndex Index = new APIIndex(Modules);
@@ -934,8 +1007,10 @@ namespace APIDocTool
 			Utility.SafeDeleteDirectoryContents(CleanDir, true);
 		}
 
-		public static bool BuildHtml(string EngineDir, string DocToolPath, string UdnPath)
+		public static bool BuildHtml(string EngineDir, string DocToolPath, string UdnPath, string HtmlDir)
 		{
+			Utility.SafeCreateDirectory(HtmlDir);
+
 			using (Process DocToolProcess = new Process())
 			{
 				DocToolProcess.StartInfo.WorkingDirectory = EngineDir;
