Move Doxygen and HTML Help Workshop to NotForLicensees.
Updated APIDocTool to point to the new locations.

diff --git a/Engine/Source/Programs/UnrealDocTool/APIDocTool/APIDocTool/Program.cs b/Engine/Source/Programs/UnrealDocTool/APIDocTool/APIDocTool/Program.cs
--- a/Engine/Source/Programs/UnrealDocTool/APIDocTool/APIDocTool/Program.cs
+++ b/Engine/Source/Programs/UnrealDocTool/APIDocTool/APIDocTool/Program.cs
@@ -362,12 +362,12 @@ namespace APIDocTool
 				if (MetadataDir == null) MetadataDir = Path.Combine(EngineDir, "Intermediate\\Documentation\\Default\\Metadata");
 
 				// Derive all the engine paths we need in one place
-				string DoxygenPath = Path.Combine(EngineDir, "Extras\\NoRedist\\Doxygen\\bin\\doxygen.exe");
+				string DoxygenPath = Path.Combine(EngineDir, "Extras\\NotForLicensees\\Doxygen\\bin\\doxygen.exe");
 				string UdnDir = Path.Combine(EngineDir, "Documentation\\Source");
 				string HtmlDir = Path.Combine(EngineDir, "Documentation\\HTML");
 				string ChmDir = Path.Combine(EngineDir, "Documentation\\CHM");
 				string DocToolPath = Path.Combine(EngineDir, "Binaries\\DotNET\\UnrealDocTool.exe");
-				string ChmCompilerPath = Path.Combine(EngineDir, "Extras\\NoRedist\\HTML Help Workshop\\hhc.exe");
+				string ChmCompilerPath = Path.Combine(EngineDir, "Extras\\NotForLicensees\\HTML Help Workshop\\hhc.exe");
 				string SettingsPath = Path.Combine(EngineDir, "Documentation\\Extras\\API\\API.ini");
 				string MetadataPath = Path.Combine(MetadataDir, "metadata.xml");
 
