Remove unused assembly Microsoft.VisualStudio.VCCodeModel to prevent builder from throwing an error

diff --git a/Engine/Source/Programs/UnrealDocTool/UnrealDocTool/UnrealDocTool/UnrealDocTool.csproj b/Engine/Source/Programs/UnrealDocTool/UnrealDocTool/UnrealDocTool/UnrealDocTool.csproj
--- a/Engine/Source/Programs/UnrealDocTool/UnrealDocTool/UnrealDocTool/UnrealDocTool.csproj
+++ b/Engine/Source/Programs/UnrealDocTool/UnrealDocTool/UnrealDocTool/UnrealDocTool.csproj
@@ -96,9 +96,6 @@
     <ApplicationIcon>DocTools.ico</ApplicationIcon>
   </PropertyGroup>
   <ItemGroup>
-    <Reference Include="Microsoft.VisualStudio.VCCodeModel, Version=10.0.0.0, Culture=neutral, PublicKeyToken=b03f5f7f11d50a3a">
-      <EmbedInteropTypes>True</EmbedInteropTypes>
-    </Reference>
     <Reference Include="Pechkin, Version=0.5.8.1, Culture=neutral, processorArchitecture=x86">
       <SpecificVersion>False</SpecificVersion>
       <HintPath>..\..\..\..\..\Binaries\DotNET\Pechkin.dll</HintPath>
