UnrealDocTranslator: Fix path passed to UnrealDocTool

diff --git a/Engine/Source/Programs/UnrealDocTool/UnrealDocTranslator/UnrealDocTranslator/UnrealDocTranslatorForm.cs b/Engine/Source/Programs/UnrealDocTool/UnrealDocTranslator/UnrealDocTranslator/UnrealDocTranslatorForm.cs
--- a/Engine/Source/Programs/UnrealDocTool/UnrealDocTranslator/UnrealDocTranslator/UnrealDocTranslatorForm.cs
+++ b/Engine/Source/Programs/UnrealDocTool/UnrealDocTranslator/UnrealDocTranslator/UnrealDocTranslatorForm.cs
@@ -650,6 +650,8 @@ namespace UnrealDocTranslator
 
         private void UdtRun(string path, bool recursive = false, UdtRunMode mode = UdtRunMode.Normal, string lang = null)
         {
+            path = new Uri(GetUnrealDocToolLocation()).MakeRelativeUri(new Uri(path)).ToString();
+
             if (currentUdtLoggingThread != null && currentUdtLoggingThread.IsAlive)
             {
                 MessageBox.Show("Conversion in progress. You have to wait until previous finishes to run another.");
