UnrealDocTranslator: Fix "Verify All" button

diff --git a/Engine/Source/Programs/UnrealDocTool/UnrealDocTranslator/UnrealDocTranslator/UnrealDocTranslatorForm.cs b/Engine/Source/Programs/UnrealDocTool/UnrealDocTranslator/UnrealDocTranslator/UnrealDocTranslatorForm.cs
--- a/Engine/Source/Programs/UnrealDocTool/UnrealDocTranslator/UnrealDocTranslator/UnrealDocTranslatorForm.cs
+++ b/Engine/Source/Programs/UnrealDocTool/UnrealDocTranslator/UnrealDocTranslator/UnrealDocTranslatorForm.cs
@@ -660,9 +660,14 @@ namespace UnrealDocTranslator
 
             if (File.Exists(path) || Directory.Exists(path))
             {
+                if (Directory.Exists(path) && recursive)
+                {
+                    path += "*";
+                }
+
                 // Relative path to UnrealDocTool, same folder.
                 var startInfo = new ProcessStartInfo(GetUnrealDocToolLocation(),
-                    String.Format("\"{0}\" " + (recursive ? "-r " : "") + GetParameters(mode), path) + (lang != null ? (" -lang=" + lang) : ""))
+                    String.Format("\"{0}\" " + GetParameters(mode), path) + (lang != null ? (" -lang=" + lang) : ""))
                                     {
                                         UseShellExecute = false,
                                         CreateNoWindow = true,
