UnrealDocTool - hack fix for doxygen crash

diff --git a/Engine/Source/Programs/UnrealDocTool/UnrealDocTool/MarkdownSharp/Doxygen/DoxygenHelper.cs b/Engine/Source/Programs/UnrealDocTool/UnrealDocTool/MarkdownSharp/Doxygen/DoxygenHelper.cs
--- a/Engine/Source/Programs/UnrealDocTool/UnrealDocTool/MarkdownSharp/Doxygen/DoxygenHelper.cs
+++ b/Engine/Source/Programs/UnrealDocTool/UnrealDocTool/MarkdownSharp/Doxygen/DoxygenHelper.cs
@@ -297,8 +297,8 @@ namespace MarkdownSharp.Doxygen
                                 // ignore compound
                                 break;
                             default:
-                                throw new InvalidOperationException(Language.Message("DoxygenCompoundKindNotSupported",
-                                                                                     compound.GetAttribute("kind", "")));
+                                //throw new InvalidOperationException(Language.Message("DoxygenCompoundKindNotSupported", compound.GetAttribute("kind", "")));
+                                break;
                         }
                     }
                 }
