UnrealDocTool - fixes for image extensions and API links

diff --git a/Engine/Source/Programs/UnrealDocTool/UnrealDocTool/MarkdownSharp/EpicMarkdown/EMImage.cs b/Engine/Source/Programs/UnrealDocTool/UnrealDocTool/MarkdownSharp/EpicMarkdown/EMImage.cs
--- a/Engine/Source/Programs/UnrealDocTool/UnrealDocTool/MarkdownSharp/EpicMarkdown/EMImage.cs
+++ b/Engine/Source/Programs/UnrealDocTool/UnrealDocTool/MarkdownSharp/EpicMarkdown/EMImage.cs
@@ -317,7 +317,7 @@ namespace MarkdownSharp.EpicMarkdown
                 try
                 {
                     var path = new EMLocalFilePath(Markdown.Unescape(url), doc, data,
-                        fileName => System.IO.Path.GetFileNameWithoutExtension(fileName) + "." + imageFormatType.ToString());
+                        fileName => System.IO.Path.GetFileNameWithoutExtension(fileName) + "." + ImageConversion.GetImageExt(imageFormatType));
 
                     if (!path.IsImage)
                     {
diff --git a/Engine/Source/Programs/UnrealDocTool/UnrealDocTool/MarkdownSharp/EpicMarkdown/EMLink.cs b/Engine/Source/Programs/UnrealDocTool/UnrealDocTool/MarkdownSharp/EpicMarkdown/EMLink.cs
--- a/Engine/Source/Programs/UnrealDocTool/UnrealDocTool/MarkdownSharp/EpicMarkdown/EMLink.cs
+++ b/Engine/Source/Programs/UnrealDocTool/UnrealDocTool/MarkdownSharp/EpicMarkdown/EMLink.cs
@@ -116,6 +116,22 @@ namespace MarkdownSharp.EpicMarkdown
         {
             try
             {
+                bool isAPILink = false;
+                string APIMember = "";
+                if (path.StartsWith("API:"))
+                {
+                    APIMember = path.Substring(4, path.Length - 4);
+                    if (Markdown.APIFileLocation.ContainsKey(APIMember))
+                    {
+                        path = "https://docs.unrealengine.com/latest/INT/" + Markdown.APIFileLocation[APIMember].Replace("\\", "/") + "/index.html";
+                    }
+                    else
+                    {
+                        content.Parse(APIMember, data);
+                        return content;
+                    }
+                    isAPILink = true;
+                }
                 var pathProvider = EMPathProvider.CreatePath(path, doc, data);
                 var errorId = -1;
 
@@ -157,20 +173,41 @@ namespace MarkdownSharp.EpicMarkdown
                     }
                     else if (pathProvider is EMLocalDocumentPath)
                     {
-                        content.Parse(
-                            data.ProcessedDocumentCache.Variables.ReplaceVariables(
-                                content.Document, "%" + path + ":title%", data),
-                            data);
+                        if (isAPILink)
+                        {
+                            content.Parse(APIMember, data);
+                        }
+                        else
+                        {
+                            content.Parse(
+                                data.ProcessedDocumentCache.Variables.ReplaceVariables(
+                                    content.Document, "%" + path + ":title%", data),
+                                data);
+                        }
                     }
                     else if (pathProvider is EMLocalFilePath)
                     {
-                        // If no link text provided use the file name.
-                        content.Parse(Regex.Replace(pathProvider.GetPath(data), @".*[\\|/](.*)", "$1"), data);
+                        if (isAPILink)
+                        {
+                            content.Parse(APIMember, data);
+                        }
+                        else
+                        {
+                            // If no link text provided use the file name.
+                            content.Parse(Regex.Replace(pathProvider.GetPath(data), @".*[\\|/](.*)", "$1"), data);
+                        }
                     }
                     else
                     {
-                        // If no link text provided use the url (remove protocol, e.g. ftp:// http://).
-                        content.Parse(Regex.Replace(pathProvider.GetPath(data), @"([^:]+://)?(.*)", "$2"), data);
+                        if (isAPILink)
+                        {
+                            content.Parse(APIMember, data);
+                        }
+                        else
+                        {
+                            // If no link text provided use the url (remove protocol, e.g. ftp:// http://).
+                            content.Parse(Regex.Replace(pathProvider.GetPath(data), @"([^:]+://)?(.*)", "$2"), data);
+                        }
                     }
                 }
 
@@ -179,7 +216,14 @@ namespace MarkdownSharp.EpicMarkdown
                     return new EMErrorElement(doc, orig, parent, errorId);
                 }
 
-                return new EMLink(doc, orig, parent, data, pathProvider, title, content);
+                if (isAPILink)
+                {
+                    return new EMLink(doc, orig, parent, data, pathProvider, APIMember, content);
+                }
+                else
+                {
+                    return new EMLink(doc, orig, parent, data, pathProvider, title, content);
+                }
             }
             catch (EMPathVerificationException e)
             {
diff --git a/Engine/Source/Programs/UnrealDocTool/UnrealDocTool/MarkdownSharp/ImageConversion.cs b/Engine/Source/Programs/UnrealDocTool/UnrealDocTool/MarkdownSharp/ImageConversion.cs
--- a/Engine/Source/Programs/UnrealDocTool/UnrealDocTool/MarkdownSharp/ImageConversion.cs
+++ b/Engine/Source/Programs/UnrealDocTool/UnrealDocTool/MarkdownSharp/ImageConversion.cs
@@ -168,6 +168,27 @@ namespace MarkdownSharp
             }
         }
 
+        public static String GetImageExt(ImageFormat formatType)
+        {
+            //Only support conversin to jpg, gif and png for web images
+            if (formatType == ImageFormat.Jpeg)
+            {
+                return "jpg";
+            }
+            if (formatType == ImageFormat.Gif)
+            {
+                return "gif";
+            }
+            if (formatType == ImageFormat.Png)
+            {
+                return "png";
+            }
+            else
+            {
+                throw new Exception(Language.Message("UnknowFormatGivenForConversion", formatType.ToString()));
+            }
+        }
+
         /// <summary>
         /// returns a codec we use for image compression
         /// </summary>
