UnrealDocTool - fix for API Links

diff --git a/Engine/Source/Programs/UnrealDocTool/UnrealDocTool/MarkdownSharp/Markdown.cs b/Engine/Source/Programs/UnrealDocTool/UnrealDocTool/MarkdownSharp/Markdown.cs
--- a/Engine/Source/Programs/UnrealDocTool/UnrealDocTool/MarkdownSharp/Markdown.cs
+++ b/Engine/Source/Programs/UnrealDocTool/UnrealDocTool/MarkdownSharp/Markdown.cs
@@ -533,7 +533,7 @@ namespace MarkdownSharp
         public static readonly Dictionary<string, string> _backslashEscapeTable;
 
         //Added dictionary to hold the API manifest information for looking up document names
-        private static readonly Dictionary<string, string> APIFileLocation = new Dictionary<string, string>();
+        public static readonly Dictionary<string, string> APIFileLocation = new Dictionary<string, string>();
 
         private readonly Dictionary<string, string> _quoteBlock = new Dictionary<string, string>();
 
@@ -1089,7 +1089,7 @@ namespace MarkdownSharp
 
                 if (translatedLanguage.Equals(data.CurrentFolderDetails.Language))
                 {
-                    linkParams.Add("selected", "selected");
+                    linkParams.Add("selected", " selected");
                 }
 
                 if(Directory.GetFiles(Path.Combine(
@@ -1097,7 +1097,7 @@ namespace MarkdownSharp
                         data.CurrentFolderDetails.CurrentFolderFromMarkdownAsTopLeaf)),
                         string.Format("*.{0}.udn", translatedLanguage)).Length == 0)
                 {
-                    linkParams.Add("disabled", "disabled");
+                    linkParams.Add("disabled", " disabled");
                 }
 
                 // Cope with top level folders having blank current folder.
