UnrealDocTool - Fix for empty directory

diff --git a/Engine/Source/Programs/UnrealDocTool/UnrealDocTool/UnrealDocTool/Program.cs b/Engine/Source/Programs/UnrealDocTool/UnrealDocTool/UnrealDocTool/Program.cs
--- a/Engine/Source/Programs/UnrealDocTool/UnrealDocTool/UnrealDocTool/Program.cs
+++ b/Engine/Source/Programs/UnrealDocTool/UnrealDocTool/UnrealDocTool/Program.cs
@@ -474,9 +474,17 @@ namespace UnrealDocTool
 
             foreach (var folder in folders)
             {
+                if (folder.Directory.GetFiles().Length == 0)
+                {
+                    continue;
+                }
                 statistics.IncrementTotal();
                 foreach (var lang in folder.Languages)
                 {
+                    if (folder.Directory.GetFiles(string.Format("*.{0}.udn", lang)).Length == 0)
+                    {
+                        continue;
+                    }
                     var fileName = folder.Directory.GetFiles(string.Format("*.{0}.udn", lang)).First().FullName;
                     var langLinks = config.LinksToAllLangs ? null : folder.Languages.ToArray();
 
@@ -886,7 +894,7 @@ namespace UnrealDocTool
 
                 var availableLangs = Directory.GetFiles("*.udn").Select(f => GetLangIdFromFileName(f.Name)).ToList();
 
-                Languages = StringListIntersection(availableLangs, requestedLangs, SupportedLanguages).ToList();
+                Languages = StringListIntersection(requestedLangs, SupportedLanguages).ToList();
             }
 
             private static string GetLangIdFromFileName(string name)
