Include plugins in API documentation. Just has a flat list of all plugin modules at the moment.

diff --git a/Engine/Source/Programs/UnrealDocTool/APIDocTool/APIDocTool/APIIndex.cs b/Engine/Source/Programs/UnrealDocTool/APIDocTool/APIDocTool/APIIndex.cs
--- a/Engine/Source/Programs/UnrealDocTool/APIDocTool/APIDocTool/APIIndex.cs
+++ b/Engine/Source/Programs/UnrealDocTool/APIDocTool/APIDocTool/APIIndex.cs
@@ -61,6 +61,7 @@ namespace APIDocTool
 			RootCategory.Categories.Add(new APIModuleCategory("Runtime"));
 			RootCategory.Categories.Add(new APIModuleCategory("Editor"));
 			RootCategory.Categories.Add(new APIModuleCategory("Developer"));
+			RootCategory.Categories.Add(new APIModuleCategory("Plugins"));
 
 			// Add all of the modules to a category
 			Dictionary<DoxygenModule, APIModuleCategory> ModuleToCategory = new Dictionary<DoxygenModule, APIModuleCategory>();
diff --git a/Engine/Source/Programs/UnrealDocTool/APIDocTool/APIDocTool/Program.cs b/Engine/Source/Programs/UnrealDocTool/APIDocTool/APIDocTool/Program.cs
--- a/Engine/Source/Programs/UnrealDocTool/APIDocTool/APIDocTool/Program.cs
+++ b/Engine/Source/Programs/UnrealDocTool/APIDocTool/APIDocTool/Program.cs
@@ -571,7 +571,7 @@ namespace APIDocTool
 			{
 				foreach(BuildModule Module in Binary.Modules)
 				{
-					if(Module.Type == "cpp" && !Module.Path.ToLowerInvariant().Contains("\\engine\\plugins\\"))
+					if(Module.Type == "cpp")
 					{
 						if(!Module.Path.ToLowerInvariant().Split('\\', '/').Any(x => InvariantExcludeDirectories.Contains(x)))
 						{
