Make the "Could not find image in given text" error spit out what text is causing the error so the translators can find the cause of the issue

diff --git a/Engine/Source/Programs/UnrealDocTool/UnrealDocTool/MarkdownSharp/EpicMarkdown/Parsers/SpanParser/EMSpanParser.cs b/Engine/Source/Programs/UnrealDocTool/UnrealDocTool/MarkdownSharp/EpicMarkdown/Parsers/SpanParser/EMSpanParser.cs
--- a/Engine/Source/Programs/UnrealDocTool/UnrealDocTool/MarkdownSharp/EpicMarkdown/Parsers/SpanParser/EMSpanParser.cs
+++ b/Engine/Source/Programs/UnrealDocTool/UnrealDocTool/MarkdownSharp/EpicMarkdown/Parsers/SpanParser/EMSpanParser.cs
@@ -79,7 +79,8 @@ namespace MarkdownSharp.EpicMarkdown.Parsers.SpanParser
 
                 if (outElement == null)
                 {
-                    throw new InvalidOperationException("Could not find image in given text.");
+                    string imageError = "Could not find image in given text: " + origin.Text;
+                    throw new InvalidOperationException(imageError);
                 }
 
                 return outElement;
