Add hpp extensions to parsing for doxygen; should prevent errors in Phya.

diff --git a/Engine/Source/Programs/UnrealDocTool/DoxygenLib/Doxygen.cs b/Engine/Source/Programs/UnrealDocTool/DoxygenLib/Doxygen.cs
--- a/Engine/Source/Programs/UnrealDocTool/DoxygenLib/Doxygen.cs
+++ b/Engine/Source/Programs/UnrealDocTool/DoxygenLib/Doxygen.cs
@@ -87,7 +87,7 @@ namespace DoxygenLib
 
 				// List of valid file patterns
 				Output.WriteLine();
-				FormatSetting(Output, "FILE_PATTERNS", new List<string> { "*.h", "*.c", "*.cpp", "*.inl", "*.inc" });
+				FormatSetting(Output, "FILE_PATTERNS", new List<string> { "*.h", "*.c", "*.hpp", "*.cpp", "*.inl", "*.inc" });
 
 				// List of tag files
 				Output.WriteLine();
