Missing defines for current compile environment.

diff --git a/Engine/Source/Programs/UnrealDocTool/APIDocTool/APIDocTool/Program.cs b/Engine/Source/Programs/UnrealDocTool/APIDocTool/APIDocTool/Program.cs
--- a/Engine/Source/Programs/UnrealDocTool/APIDocTool/APIDocTool/Program.cs
+++ b/Engine/Source/Programs/UnrealDocTool/APIDocTool/APIDocTool/Program.cs
@@ -110,6 +110,8 @@ namespace APIDocTool
 			"NAVOCTREE_CONTAINS_COLLISION_DATA=1",
 			"SOURCE_CONTROL_WITH_SLATE=1",
 			"MATCHMAKING_HACK_FOR_EGP_IE_HOSTING=1",
+			"USE_REMOTE_INTEGRATION=1",
+			"WITH_FANCY_TEXT=1",
 
 			// Online subsystems
 			"PACKAGE_SCOPE:=protected",
