Fix UnrealDocTool incorrectly treating escaped characters in URLs as variable names.

diff --git a/Engine/Source/Programs/UnrealDocTool/UnrealDocTool/MarkdownSharp/Preprocessor/VariableManager.cs b/Engine/Source/Programs/UnrealDocTool/UnrealDocTool/MarkdownSharp/Preprocessor/VariableManager.cs
--- a/Engine/Source/Programs/UnrealDocTool/UnrealDocTool/MarkdownSharp/Preprocessor/VariableManager.cs
+++ b/Engine/Source/Programs/UnrealDocTool/UnrealDocTool/MarkdownSharp/Preprocessor/VariableManager.cs
@@ -15,7 +15,7 @@ namespace MarkdownSharp.Preprocessor
 
     public class VariableManager
     {
-        public static readonly Regex InLineVariableOrMetadata = new Regex(@"(?<=^[ ]{0,3}(\S.*)?)((\%(?<pathWithColon>(?<path>[^%: \s]*?):)?(?<variableName>([_ \(\):a-zA-Z0-9\<\>*;]+?))\%))", RegexOptions.IgnorePatternWhitespace | RegexOptions.Multiline | RegexOptions.Compiled);
+        public static readonly Regex InLineVariableOrMetadata = new Regex(@"(?<=^[ ]{0,3}(\S.*)?)(?<!\(http://[^)]+)((\%(?<pathWithColon>(?<path>[^%: \s]*?):)?(?<variableName>([_ \(\):a-zA-Z0-9\<\>*;]+?))\%))", RegexOptions.IgnorePatternWhitespace | RegexOptions.Multiline | RegexOptions.Compiled);
         private static readonly Regex VariableDefinitionPattern = new Regex(@"
                 (
                     ^\[VAR:(?<variableName>[_a-zA-Z0-9]+?)\]$\n
