Improvements to API documentation.
- Stop treating the DEPRECATED macro like a function definition
- Add function parameter names into index pages
- Fix resolving {{typedef-type}} markup
- Remove hyphens (causing bullet points) after [REDACTED_HANDLE] directive
- Remove deprecated functions and variables
- Separate functions which are overriden from different classes into separate sections.
- Remove the separate templates for each object type.
- Remove markup section, but link to hand-written specifier descriptions instead.

diff --git a/Engine/Source/Programs/UnrealDocTool/APIDocTool/APIDocTool/APIConstant.cs b/Engine/Source/Programs/UnrealDocTool/APIDocTool/APIDocTool/APIConstant.cs
--- a/Engine/Source/Programs/UnrealDocTool/APIDocTool/APIDocTool/APIConstant.cs
+++ b/Engine/Source/Programs/UnrealDocTool/APIDocTool/APIDocTool/APIConstant.cs
@@ -32,38 +32,28 @@ namespace APIDocTool
 			{
 				Writer.WritePageHeader(Name, PageCrumbs, BriefDescription);
 
-				Writer.EnterTag("[OBJECT:Constant]");
-
-				Writer.EnterTag("[PARAM:briefdesc]");
-				if (!Utility.IsNullOrWhitespace(BriefDescription))
-				{
-					Writer.WriteLine(BriefDescription);
-				}
-				Writer.LeaveTag("[/PARAM]");
-
 				// Write the syntax
-				Writer.EnterTag("[PARAM:syntax]");
 				Writer.EnterSection("syntax", "Syntax");
 				WriteDefinition(Writer);
 				Writer.LeaveSection();
-				Writer.LeaveTag("[/PARAM]");
 
 				// Write the description
-				Writer.EnterTag("[PARAM:description]");
-				if (!Utility.IsNullOrWhitespace(FullDescription) && FullDescription != BriefDescription)
+				if (!Utility.IsNullOrWhitespace(BriefDescription) || !Utility.IsNullOrWhitespace(FullDescription))
 				{
 					Writer.EnterSection("description", "Remarks");
-					Writer.WriteLine(FullDescription);
+					if (!Utility.IsNullOrWhitespace(BriefDescription) && FullDescription != BriefDescription)
+					{
+						Writer.WriteLine(BriefDescription);
+					}
+					if(!Utility.IsNullOrWhitespace(FullDescription))
+					{
+						Writer.WriteLine(FullDescription);
+					}
 					Writer.LeaveSection();
 				}
-				Writer.LeaveTag("[/PARAM]");
 
 				// Write the references
-				Writer.EnterParam("references");
 				WriteReferencesSection(Writer, Entity);
-				Writer.LeaveParam();
-
-				Writer.LeaveTag("[/OBJECT]");
 			}
 		}
 
diff --git a/Engine/Source/Programs/UnrealDocTool/APIDocTool/APIDocTool/APIEnum.cs b/Engine/Source/Programs/UnrealDocTool/APIDocTool/APIDocTool/APIEnum.cs
--- a/Engine/Source/Programs/UnrealDocTool/APIDocTool/APIDocTool/APIEnum.cs
+++ b/Engine/Source/Programs/UnrealDocTool/APIDocTool/APIDocTool/APIEnum.cs
@@ -95,51 +95,28 @@ namespace APIDocTool
 			{
 				Writer.WritePageHeader(AnnotateAnonymousNames(Name), PageCrumbs, BriefDescription);
 
-				Writer.EnterTag("[OBJECT:Enum]");
-
-				Writer.EnterTag("[PARAM:briefdesc]");
-				if (!Utility.IsNullOrWhitespace(BriefDescription))
-				{
-					Writer.WriteLine(BriefDescription);
-				}
-				Writer.LeaveTag("[/PARAM]");
-
 				// Write the syntax
-				Writer.EnterTag("[PARAM:syntax]");
 				Writer.EnterSection("syntax", "Syntax");
 				WriteDefinition(Writer);
 				Writer.LeaveSection();
-				Writer.LeaveTag("[/PARAM]");
-
-				// Write the metadata
-				Writer.EnterTag("[PARAM:meta]");
-				if (MetadataDirective != null)
-				{
-					MetadataDirective.WriteListSection(Writer, "metadata", "Metadata", MetadataLookup.EnumTags);
-				}
-				Writer.LeaveTag("[/PARAM]");
 
 				// Write the enum values
-				Writer.EnterTag("[PARAM:values]");
-				Writer.WriteListSection("values", "Values", "Name", "Description", Values.Select(x => x.GetListItem()));
-				Writer.LeaveTag("[/PARAM]");
-
-				// Write the description
-				Writer.EnterTag("[PARAM:description]");
-				if (!Utility.IsNullOrWhitespace(FullDescription) && FullDescription != BriefDescription)
+				if (!Utility.IsNullOrWhitespace(BriefDescription) || !Utility.IsNullOrWhitespace(FullDescription))
 				{
 					Writer.EnterSection("description", "Remarks");
-					Writer.WriteLine(FullDescription);
+					if (!Utility.IsNullOrWhitespace(BriefDescription) && FullDescription != BriefDescription)
+					{
+						Writer.WriteLine(BriefDescription);
+					}
+					if (!Utility.IsNullOrWhitespace(FullDescription) && FullDescription != BriefDescription)
+					{
+						Writer.WriteLine(FullDescription);
+					}
 					Writer.LeaveSection();
 				}
-				Writer.LeaveTag("[/PARAM]");
 
 				// Write the references
-				Writer.EnterTag("[PARAM:references]");
 				WriteReferencesSection(Writer, Entity);
-				Writer.LeaveTag("[/PARAM]");
-
-				Writer.LeaveTag("[/OBJECT]");
 			}
         }
 
diff --git a/Engine/Source/Programs/UnrealDocTool/APIDocTool/APIDocTool/APIFunction.cs b/Engine/Source/Programs/UnrealDocTool/APIDocTool/APIDocTool/APIFunction.cs
--- a/Engine/Source/Programs/UnrealDocTool/APIDocTool/APIDocTool/APIFunction.cs
+++ b/Engine/Source/Programs/UnrealDocTool/APIDocTool/APIDocTool/APIFunction.cs
@@ -219,26 +219,26 @@ namespace APIDocTool
 			get { return APIMember.ConvertToMarkdown(Node.SelectSingleNode("type")); }
 		}
 
-		public string GetAbbreviatedTypeText()
+		public string GetAbbreviatedDefinition()
 		{
-			string TypeText = Type;
-			if (TypeText.StartsWith("class "))
+			string Text = Definition;
+			if (Text.StartsWith("class "))
 			{
-				TypeText = TypeText.Substring(6);
+				Text = Text.Substring(6);
 			}
-			else if (TypeText.StartsWith("const class "))
+			else if (Text.StartsWith("const class "))
 			{
-				TypeText = "const " + TypeText.Substring(12);
+				Text = "const " + Text.Substring(12);
 			}
-			else if (TypeText.StartsWith("struct "))
+			else if (Text.StartsWith("struct "))
 			{
-				TypeText = TypeText.Substring(7);
+				Text = Text.Substring(7);
 			}
-			else if (TypeText.StartsWith("const struct "))
+			else if (Text.StartsWith("const struct "))
 			{
-				TypeText = "const " + TypeText.Substring(13);
+				Text = "const " + Text.Substring(13);
 			}
-			return TypeText;
+			return Text;
 		}
 
 		public string Definition
@@ -265,7 +265,7 @@ namespace APIDocTool
 			}
 
 			XmlNode DescriptionNode = Node.SelectSingleNode("parameterdescription");
-			Description = DescriptionNode.InnerText;
+			Description = DescriptionNode.InnerText.TrimStart('-', ' ');
 		}
 
 		public UdnListItem GetListItem()
@@ -342,6 +342,26 @@ namespace APIDocTool
 			get { APIFunctionGroup ParentFunctionGroup = Parent as APIFunctionGroup; return (ParentFunctionGroup != null) ? ParentFunctionGroup.FullName : base.FullName; }
 		}
 
+		public APIFunction GetBaseImplementation()
+		{
+			APIFunction BaseFunction = this;
+			while(BaseFunction.Reimplements.Count > 0)
+			{
+				BaseFunction = BaseFunction.Reimplements[0];
+			}
+			return BaseFunction;
+		}
+
+		public bool IsDeprecated()
+		{
+			return Name.EndsWith("_DEPRECATED");
+		}
+
+		public bool IsExecFunction()
+		{
+			return Name.Length >= 5 && Name.StartsWith("exec") && Char.IsUpper(Name[4]);
+		}
+
 		public override void AddToManifest(UdnManifest Manifest)
 		{
 			if (FunctionType != APIFunctionType.UnaryOperator && FunctionType != APIFunctionType.BinaryOperator)
@@ -520,7 +540,7 @@ namespace APIDocTool
 			List<string> Lines = new List<string>();
 			if(MetadataDirective != null)
 			{
-				Lines.Add(MetadataDirective.ToMarkdown());
+				Lines.AddRange(MetadataDirective.ToMarkdown());
 			}
 			if (!Utility.IsNullOrWhitespace(TemplateSignature))
 			{
@@ -555,18 +575,6 @@ namespace APIDocTool
 			WriteNestedSimpleCode(Writer, Lines);
 		}
 
-		private void WriteSource(UdnWriter Writer)
-		{
-			Writer.EnterTag("[REGION:simplecode_api]");
-			foreach (string SourceLine in SourceLines)
-			{
-				string EscapedLine = Utility.EscapeText(SourceLine);
-				EscapedLine = EscapedLine.Replace(" ", UdnWriter.Space);
-				Writer.WriteLine(EscapedLine + "  ");
-			}
-			Writer.LeaveTag("[/REGION]");
-		}
-
 		private void WriteSourceSection(UdnWriter Writer)
 		{
 			if(Entity.BodyFile != null)
@@ -645,18 +653,8 @@ namespace APIDocTool
 			using (UdnWriter Writer = new UdnWriter(OutputPath))
 			{
 				Writer.WritePageHeader(FullName, PageCrumbs, BriefDescription);
-				Writer.EnterTag("[OBJECT:Function]");
-
-				// Write the brief description
-				Writer.EnterTag("[PARAM:briefdesc]");
-				if (!Utility.IsNullOrWhitespace(BriefDescription) && BriefDescription != FullDescription)
-				{
-					Writer.WriteLine(BriefDescription);
-				}
-				Writer.LeaveTag("[/PARAM]");
 
 				// Write the warnings
-				Writer.EnterTag("[PARAM:warnings]");
 				if (Warnings.Count > 0)
 				{
 					Writer.EnterTag("[REGION:warning]");
@@ -667,10 +665,8 @@ namespace APIDocTool
 					}
 					Writer.LeaveTag("[/REGION]");
 				}
-				Writer.LeaveTag("[/PARAM]");
 
 				// Write the virtual hierarchy
-				Writer.EnterTag("[PARAM:overrides]");
 				if (HierarchyNode != null)
 				{
 					Writer.EnterSection("overrides", "Override Hierarchy");
@@ -679,63 +675,56 @@ namespace APIDocTool
 					Writer.LeaveTag("[/REGION]");
 					Writer.LeaveSection();
 				}
-				Writer.LeaveTag("[/PARAM]");
 
 				// Write the syntax
-				Writer.EnterTag("[PARAM:syntax]");
 				Writer.EnterSection("syntax", "Syntax");
 				WriteSyntax(Writer);
 				Writer.LeaveSection();
-				Writer.LeaveTag("[/PARAM]");
+
+				// Write the description
+				if (!Utility.IsNullOrWhitespace(BriefDescription) || !Utility.IsNullOrWhitespace(FullDescription))
+				{
+					Writer.EnterSection("description", "Remarks");
+					if (!Utility.IsNullOrWhitespace(BriefDescription) && BriefDescription != FullDescription)
+					{
+						Writer.WriteLine(BriefDescription);
+					}
+					if (!Utility.IsNullOrWhitespace(FullDescription))
+					{
+						Writer.WriteLine(BriefDescription);
+					}
+					Writer.LeaveSection();
+				}
+
+				// Enter the body
+				Writer.EnterRegion("syntax");
 
 				// Write the return type
-				Writer.EnterTag("[PARAM:returns]");
 				if(!Utility.IsNullOrWhitespace(ReturnDescription))
 				{
 					Writer.EnterSection("returns", "Returns");
 					Writer.WriteLine(ReturnDescription);
 					Writer.LeaveSection();
 				}
-				Writer.LeaveTag("[/PARAM]");
 
 				// Write the parameters
-				Writer.EnterTag("[PARAM:params]");
 				if (ParameterSummaries.Count > 0)
 				{
 					Writer.WriteListSection("params", "Parameters", "Parameter", "Description", ParameterSummaries.Select(x => x.GetListItem()));
 				}
-				Writer.LeaveTag("[/PARAM]");
 
-				// Write the metadata
-				Writer.EnterTag("[PARAM:meta]");
-				if (MetadataDirective != null)
-				{
-					MetadataDirective.WriteListSection(Writer, "metadata", "Metadata", MetadataLookup.FunctionTags);
-				}
-				Writer.LeaveTag("[/PARAM]");
-
-				// Write the description
-				Writer.EnterTag("[PARAM:description]");
-				if (!Utility.IsNullOrWhitespace(FullDescription))
-				{
-					Writer.EnterSection("description", "Remarks");
-					Writer.WriteLine(FullDescription);
-					Writer.LeaveSection();
-				}
-				Writer.LeaveTag("[/PARAM]");
+				// Leave the body
+				Writer.LeaveRegion();
 
 				// Write the marshalling struct
-				Writer.EnterTag("[PARAM:marshalling]");
 				if (EventParameters != null)
 				{
 					Writer.EnterSection("marshalling", "Marshalling");
 					Writer.WriteLine("May be called as an event using [{0}]({1}) as parameter frame.", EventParameters.Name, EventParameters.LinkPath);
 					Writer.LeaveSection();
 				}
-				Writer.LeaveTag("[/PARAM]");
 
 				// Write the template specializations
-				Writer.EnterTag("[PARAM:specializations]");
 				if(TemplateSpecializations.Count > 0)
 				{
 					Writer.EnterSection("specializations", "Specializations");
@@ -745,27 +734,12 @@ namespace APIDocTool
 					}
 					Writer.LeaveSection();
 				}
-				Writer.LeaveTag("[/PARAM]");
 
 				// Write the source
-				Writer.EnterTag("[PARAM:source]");
 				WriteSourceSection(Writer);
-/*
-
-				if (SourceLines != null)
-				{
-					Writer.EnterSection("source", "Source");
-					WriteSource(Writer);
-					Writer.LeaveSection();
-				}
-*/				Writer.LeaveTag("[/PARAM]");
 
 				// Write the reference info
-				Writer.EnterTag("[PARAM:references]");
 				WriteReferencesSection(Writer, Entity);
-				Writer.LeaveTag("[/PARAM]");
-
-				Writer.LeaveTag("[/OBJECT]");
 			}
         }
 
@@ -803,7 +777,8 @@ namespace APIDocTool
 				for (int Idx = 0; Idx < Parameters.Count; Idx++)
 				{
 					string Separator = (Idx + 1 == Parameters.Count) ? "" : ",";
-					Writer.WriteLine(UdnWriter.TabSpaces + Parameters[Idx].GetAbbreviatedTypeText() + Separator + "  ");
+					string Definition = Markdown.Truncate(Parameters[Idx].GetAbbreviatedDefinition(), 35, "...");
+					Writer.WriteLine(UdnWriter.TabSpaces + Definition + Separator + "  ");
 				}
 				Writer.WriteEscapedLine(")  ");
 			}
diff --git a/Engine/Source/Programs/UnrealDocTool/APIDocTool/APIDocTool/APIFunctionGroup.cs b/Engine/Source/Programs/UnrealDocTool/APIDocTool/APIDocTool/APIFunctionGroup.cs
--- a/Engine/Source/Programs/UnrealDocTool/APIDocTool/APIDocTool/APIFunctionGroup.cs
+++ b/Engine/Source/Programs/UnrealDocTool/APIDocTool/APIDocTool/APIFunctionGroup.cs
@@ -47,6 +47,11 @@ namespace APIDocTool
 			get { return Children.OfType<APIFunction>().OrderBy(x => x.LinkPath); }
 		}
 
+		public bool IsDeprecated()
+		{
+			return Name.EndsWith("_DEPRECATED");
+		}
+
 		public override void Link()
 		{
 		}
diff --git a/Engine/Source/Programs/UnrealDocTool/APIDocTool/APIDocTool/APIRecord.cs b/Engine/Source/Programs/UnrealDocTool/APIDocTool/APIDocTool/APIRecord.cs
--- a/Engine/Source/Programs/UnrealDocTool/APIDocTool/APIDocTool/APIRecord.cs
+++ b/Engine/Source/Programs/UnrealDocTool/APIDocTool/APIDocTool/APIRecord.cs
@@ -275,7 +275,7 @@ namespace APIDocTool
 
 			if (MetadataDirective != null)
 			{
-				Lines.Add(MetadataDirective.ToMarkdown());
+				Lines.AddRange(MetadataDirective.ToMarkdown());
 			}
 			if (TemplateSignature != null)
 			{
@@ -283,7 +283,7 @@ namespace APIDocTool
 			}
 			if(BaseDefinitions.Count == 0)
 			{
-				Lines.Add(Definition);
+				Lines.Add("<span style=\"color:blue\">" + Definition + "</span>");
 			}
 			else if(BaseDefinitions.Count == 1)
 			{
@@ -320,23 +320,28 @@ namespace APIDocTool
 			return false;
 		}
 
+		private void FindAllBaseRecords(List<APIRecord> Records)
+		{
+			if(!Records.Contains(this))
+			{
+				Records.Add(this);
+				foreach(KeyValuePair<XmlNode, APIRecord> BaseRecord in BaseRecords)
+				{
+					if(BaseRecord.Value != null)
+					{
+						BaseRecord.Value.FindAllBaseRecords(Records);
+					}
+				}
+			}
+		}
+
 		public override void WritePage(UdnManifest Manifest, string OutputPath)
         {
 			using (UdnWriter Writer = new UdnWriter(OutputPath))
 			{
 				Writer.WritePageHeader(Name, PageCrumbs, BriefDescription);
-				Writer.EnterTag("[OBJECT:Class]");
-
-				// Write the brief description
-				Writer.WriteLine("[PARAM:briefdesc]");
-				if (!Utility.IsNullOrWhitespace(BriefDescription))
-				{
-					Writer.WriteLine(BriefDescription);
-				}
-				Writer.WriteLine("[/PARAM]");
 
 				// Write the hierarchy
-				Writer.EnterTag("[PARAM:hierarchy]");
 				if (HierarchyNode != null)
 				{
 					Writer.EnterSection("hierarchy", "Inheritance Hierarchy");
@@ -345,34 +350,40 @@ namespace APIDocTool
 					Writer.LeaveTag("[/REGION]");
 					Writer.LeaveSection();
 				}
-				Writer.LeaveTag("[/PARAM]");
 
 				// Write the record definition
-				Writer.EnterTag("[PARAM:syntax]");
 				if (!Utility.IsNullOrWhitespace(Definition))
 				{
 					Writer.EnterSection("syntax", "Syntax");
 					WriteDefinition(Writer);
 					Writer.LeaveSection();
 				}
-				Writer.LeaveTag("[/PARAM]");
 
-				// Write the metadata section
-				Writer.WriteLine("[PARAM:meta]");
-				if (MetadataDirective != null)
+				// Write the class description
+				if (!Utility.IsNullOrWhitespace(BriefDescription) || !Utility.IsNullOrWhitespace(FullDescription))
 				{
-					MetadataDirective.WriteListSection(Writer, "metadata", "Metadata", MetadataLookup.ClassTags);
+					Writer.EnterSection("description", "Remarks");
+					if (!Utility.IsNullOrWhitespace(BriefDescription))
+					{
+						Writer.WriteLine(BriefDescription);
+					}
+					if (!Utility.IsNullOrWhitespace(FullDescription))
+					{
+						Writer.WriteLine(FullDescription);
+					}
+					Writer.LeaveSection();
 				}
-				Writer.WriteLine("[/PARAM]");
+
+				// Write the main body section
+				Writer.EnterRegion("syntax");
 
 				// Build a list of all the functions
 				List<APIFunction> AllFunctions = new List<APIFunction>();
-				AllFunctions.AddRange(Children.OfType<APIFunction>().Where(x => x.Protection != APIProtection.Private));
-				AllFunctions.AddRange(Children.OfType<APIFunctionGroup>().SelectMany(x => x.Children.OfType<APIFunction>()).Where(x => x.Protection != APIProtection.Private));
+				AllFunctions.AddRange(Children.OfType<APIFunction>().Where(x => x.Protection != APIProtection.Private && !x.IsDeprecated()));
+				AllFunctions.AddRange(Children.OfType<APIFunctionGroup>().SelectMany(x => x.Children.OfType<APIFunction>()).Where(x => x.Protection != APIProtection.Private && !x.IsDeprecated()));
 				AllFunctions.Sort((x, y) => String.Compare(x.Name, y.Name));
 
 				// Write all the specializations
-				Writer.EnterTag("[PARAM:specializations]");
 				if (TemplateSpecializations.Count > 0)
 				{
 					Writer.EnterSection("specializations", "Specializations");
@@ -382,89 +393,78 @@ namespace APIDocTool
 					}
 					Writer.LeaveSection();
 				}
-				Writer.LeaveTag("[/PARAM]");
 
-				// Write all the typedefs
-				Writer.EnterTag("[PARAM:typedefs]");
-				Writer.WriteListSection("typedefs", "Typedefs", "Name", "Description", Children.OfType<APITypeDef>().Select(x => x.GetListItem()));
-				Writer.LeaveTag("[/PARAM]");
+				// Write all the variables
+				Writer.WriteListSection("variables", "Variables", "Name", "Description", Children.OfType<APIVariable>().Where(x => x.Protection != APIProtection.Private && !x.IsDeprecated()).OrderBy(x => x.Name).Select(x => x.GetListItem()));
 
 				// Write all the constructors
-				Writer.EnterTag("[PARAM:constructors]");
 				if (!APIFunction.WriteListSection(Writer, "constructor", "Constructors", AllFunctions.Where(x => x.FunctionType == APIFunctionType.Constructor).OrderBy(x => x.LinkPath)) && HasAnyPrivateFunction(Name))
 				{
 					Writer.EnterSection("constructor", "Constructors");
 					Writer.WriteLine("No constructors are accessible with public or protected access.");
 					Writer.LeaveSection();
 				}
-				Writer.LeaveTag("[/PARAM]");
 
 				// Write all the destructors
-				Writer.EnterTag("[PARAM:destructors]");
 				if (!APIFunction.WriteListSection(Writer, "destructor", "Destructors", AllFunctions.Where(x => x.FunctionType == APIFunctionType.Destructor)) && HasAnyPrivateFunction("~" + Name))
 				{
 					Writer.EnterSection("destructors", "Destructors");
 					Writer.WriteLine("No destructors are accessible with public or protected access.");
 					Writer.LeaveSection();
 				}
-				Writer.LeaveTag("[/PARAM]");
 
-				// Write all the enums
-				Writer.EnterTag("[PARAM:enums]");
-				Writer.WriteListSection("enums", "Enums", "Name", "Description", Children.OfType<APIEnum>().OrderBy(x => x.Name).Select(x => x.GetListItem()));
-				Writer.LeaveTag("[/PARAM]");
-
-				// Write all the inner structures
-				Writer.EnterTag("[PARAM:classes]");
-				Writer.WriteListSection("classes", "Classes", "Name", "Description", Children.OfType<APIRecord>().Where(x => x.Protection != APIProtection.Private).OrderBy(x => x.Name).Select(x => x.GetListItem()));
-				Writer.LeaveTag("[/PARAM]");
+				// Find a list of base classes
+				List<APIRecord> AllBaseRecords = new List<APIRecord>();
+				FindAllBaseRecords(AllBaseRecords);
 
-				// Write all the constants
-				Writer.EnterTag("[PARAM:constants]");
-				Writer.WriteListSection("constants", "Constants", "Name", "Description", Children.OfType<APIConstant>().Select(x => x.GetListItem()));
-				Writer.LeaveTag("[/PARAM]");
-
-				// Write all the variables
-				Writer.EnterTag("[PARAM:variables]");
-				Writer.WriteListSection("variables", "Variables", "Name", "Description", Children.OfType<APIVariable>().Where(x => x.Protection != APIProtection.Private).OrderBy(x => x.Name).Select(x => x.GetListItem()));
-				Writer.LeaveTag("[/PARAM]");
+				// Build a list of functions for each base record
+				List<APIFunction>[] AllBaseFunctions = AllBaseRecords.Select(x => new List<APIFunction>()).ToArray();
+				foreach(APIFunction Function in AllFunctions.Where(x =>x.FunctionType == APIFunctionType.Normal && !x.IsExecFunction()))
+				{
+					int BaseRecordIdx = AllBaseRecords.IndexOf(Function.GetBaseImplementation().FindParent<APIRecord>());
+					AllBaseFunctions[Math.Max(0, BaseRecordIdx)].Add(Function);
+				}
 
 				// Write the functions
-				Writer.EnterTag("[PARAM:methods]");
-				APIFunction.WriteListSection(Writer, "methods", "Methods", AllFunctions.Where(x => x.FunctionType == APIFunctionType.Normal));
-				Writer.LeaveTag("[/PARAM]");
+				for (int Idx = 0; Idx < AllBaseFunctions.Length; Idx++)
+				{
+					List<APIFunction> BaseFunctions = AllBaseFunctions[Idx];
+					if (BaseFunctions.Count > 0)
+					{
+						string Id = String.Format("methods_{0}", Idx);
+						string Label = (Idx == 0) ? "Methods" : String.Format("Overridden from {0}", AllBaseRecords[Idx].Name);
+						APIFunction.WriteListSection(Writer, Id, Label, AllBaseFunctions[Idx]);
+					}
+				}
 
 				// Write the operator list
-				Writer.EnterTag("[PARAM:operators]");
 				APIFunction.WriteListSection(Writer, "operators", "Operators", AllFunctions.Where(x => x.FunctionType == APIFunctionType.UnaryOperator || x.FunctionType == APIFunctionType.BinaryOperator));
-				Writer.LeaveTag("[/PARAM]");
 
-				// Write class description
-				Writer.EnterTag("[PARAM:description]");
-				if (!Utility.IsNullOrWhitespace(FullDescription) && FullDescription != BriefDescription)
-				{
-					Writer.EnterSection("description", "Remarks");
-					Writer.WriteLine(FullDescription);
-					Writer.LeaveSection();
-				}
-				Writer.LeaveTag("[/PARAM]");
+				// Write all the inner structures
+				Writer.WriteListSection("classes", "Classes", "Name", "Description", Children.OfType<APIRecord>().Where(x => x.Protection != APIProtection.Private).OrderBy(x => x.Name).Select(x => x.GetListItem()));
+
+				// Write all the enums
+				Writer.WriteListSection("enums", "Enums", "Name", "Description", Children.OfType<APIEnum>().OrderBy(x => x.Name).Select(x => x.GetListItem()));
+
+				// Write all the typedefs
+				Writer.WriteListSection("typedefs", "Typedefs", "Name", "Description", Children.OfType<APITypeDef>().Select(x => x.GetListItem()));
+
+				// Write all the constants
+				Writer.WriteListSection("constants", "Constants", "Name", "Description", Children.OfType<APIConstant>().Select(x => x.GetListItem()));
+
+				// Leave the body
+				Writer.LeaveRegion();
 
 				// Write the marshalling parameters
-				Writer.EnterTag("[PARAM:marshalling]");
 				if (DelegateEventParameters != null)
 				{
 					Writer.EnterSection("marshalling", "Marshalling");
 					Writer.WriteLine("Parameters are marshalled using [{0}]({1})", DelegateEventParameters.FullName, DelegateEventParameters.LinkPath);
 					Writer.LeaveSection();
 				}
-				Writer.LeaveTag("[/PARAM]");
 
 				// Write the references
-				Writer.EnterTag("[PARAM:references]");
 				WriteReferencesSection(Writer, Entity);
-				Writer.LeaveTag("[/PARAM]");
-
-				Writer.LeaveTag("[/OBJECT]");
 			}
         }
 
diff --git a/Engine/Source/Programs/UnrealDocTool/APIDocTool/APIDocTool/APITypeDef.cs b/Engine/Source/Programs/UnrealDocTool/APIDocTool/APIDocTool/APITypeDef.cs
--- a/Engine/Source/Programs/UnrealDocTool/APIDocTool/APIDocTool/APITypeDef.cs
+++ b/Engine/Source/Programs/UnrealDocTool/APIDocTool/APIDocTool/APITypeDef.cs
@@ -44,7 +44,7 @@ namespace APIDocTool
 
 		public void ReplaceMarkdownVariables(ref string Text)
 		{
-			Text = Text.Replace("&#123;&#123; typedef-type &#125;&#125;", Type);
+			Text = Text.Replace("&#123;&#123; typedef&#45;type &#125;&#125;", Type);
 		}
 
 		public override void WritePage(UdnManifest Manifest, string Path)
@@ -53,40 +53,28 @@ namespace APIDocTool
 			{
 				Writer.WritePageHeader(Name, PageCrumbs, BriefDescription);
 
-				Writer.EnterTag("[OBJECT:Typedef]");
-
-				// Write the brief description
-				Writer.EnterTag("[PARAM:briefdesc]");
-				if (!Utility.IsNullOrWhitespace(BriefDescription) && BriefDescription != FullDescription)
-				{
-					Writer.WriteLine(BriefDescription);
-				}
-				Writer.LeaveTag("[/PARAM]");
-
 				// Write the type
-				Writer.EnterTag("[PARAM:type]");
 				Writer.EnterSection("type", "Type");
 				WriteNestedSimpleCode(Writer, new List<string> { "typedef " + Type + " " + Name });
 				Writer.LeaveSection();
-				Writer.LeaveTag("[/PARAM]");
 
 				// Write the description
-				Writer.EnterTag("[PARAM:description]");
-				if (!Utility.IsNullOrWhitespace(FullDescription))
+				if (!Utility.IsNullOrWhitespace(BriefDescription) || !Utility.IsNullOrWhitespace(FullDescription))
 				{
 					Writer.EnterSection("description", "Remarks");
-					Writer.WriteLine(FullDescription);
+					if (!Utility.IsNullOrWhitespace(BriefDescription) && BriefDescription != FullDescription)
+					{
+						Writer.WriteLine(BriefDescription);
+					}
+					if (!Utility.IsNullOrWhitespace(FullDescription))
+					{
+						Writer.WriteLine(FullDescription);
+					}
 					Writer.LeaveSection();
 				}
-				Writer.LeaveTag("[/PARAM]");
 
 				// Write the references
-				Writer.EnterTag("[PARAM:references]");
 				WriteReferencesSection(Writer, Entity);
-				Writer.LeaveTag("[/PARAM]");
-
-				// Leave the object tag
-				Writer.LeaveTag("[/OBJECT]");
 			}
 		}
     }
diff --git a/Engine/Source/Programs/UnrealDocTool/APIDocTool/APIDocTool/APIVariable.cs b/Engine/Source/Programs/UnrealDocTool/APIDocTool/APIDocTool/APIVariable.cs
--- a/Engine/Source/Programs/UnrealDocTool/APIDocTool/APIDocTool/APIVariable.cs
+++ b/Engine/Source/Programs/UnrealDocTool/APIDocTool/APIDocTool/APIVariable.cs
@@ -40,6 +40,11 @@ namespace APIDocTool
 			AddRefLink(Node.Attributes["id"].InnerText, this);
 		}
 
+		public bool IsDeprecated()
+		{
+			return Name.EndsWith("_DEPRECATED");
+		}
+
 		public override void Link()
 		{
 			// Parse the metadata
@@ -81,24 +86,33 @@ namespace APIDocTool
 			IsStatic = Node.Attributes.GetNamedItem("static").InnerText == "yes";
         }
 
+		private void WriteDefinition(UdnWriter Writer)
+		{
+			List<string> Lines = new List<string>();
+
+			if (MetadataDirective != null)
+			{
+				Lines.AddRange(MetadataDirective.ToMarkdown());
+			}
+
+			StringBuilder Definition = new StringBuilder();
+			if (IsStatic) Definition.Append("static ");
+			if (IsMutable) Definition.Append("mutable ");
+			Definition.Append(Type + " " + Name);
+			if (Bitfield != "") Definition.Append(": " + Bitfield);
+			Definition.Append("  ");
+			Lines.Add(Definition.ToString());
+
+			WriteNestedSimpleCode(Writer, Lines);
+		}
+
 		public override void WritePage(UdnManifest Manifest, string OutputPath)
         {
 			using (UdnWriter Writer = new UdnWriter(OutputPath))
 			{
 				Writer.WritePageHeader(Name, PageCrumbs, BriefDescription);
 
-				Writer.EnterTag("[OBJECT:Variable]");
-
-				// Write the brief description
-				Writer.EnterTag("[PARAM:briefdesc]");
-				if (!Utility.IsNullOrWhitespace(BriefDescription) && BriefDescription != FullDescription)
-				{
-					Writer.WriteLine(BriefDescription);
-				}
-				Writer.LeaveTag("[/PARAM]");
-
 				// Write the warnings
-				Writer.EnterTag("[PARAM:warnings]");
 				if (Warnings.Count > 0)
 				{
 					Writer.EnterTag("[REGION:warning]");
@@ -109,42 +123,26 @@ namespace APIDocTool
 					}
 					Writer.LeaveTag("[/REGION]");
 				}
-				Writer.LeaveTag("[/PARAM]");
 
 				// Write the syntax
-				Writer.EnterTag("[PARAM:syntax]");
 				Writer.EnterSection("syntax", "Syntax");
-				Writer.EnterTag("[REGION:simplecode_api]");
-				if (MetadataDirective != null) Writer.WriteLine(MetadataDirective.ToMarkdown() + "  ");
-				if (IsStatic) Writer.Write("static ");
-				if (IsMutable) Writer.Write("mutable ");
-				Writer.Write(Type + " " + Name);
-				if (Bitfield != "") Writer.Write(": " + Bitfield);
-				Writer.WriteLine("  ");
-				Writer.LeaveTag("[/REGION]");
+				WriteDefinition(Writer);
 				Writer.LeaveSection();
-				Writer.LeaveTag("[/PARAM]");
-
-				// Write the metadata
-				Writer.EnterTag("[PARAM:meta]");
-				if (MetadataDirective != null)
-				{
-					MetadataDirective.WriteListSection(Writer, "metadata", "Metadata", MetadataLookup.PropertyTags);
-				}
-				Writer.LeaveTag("[/PARAM]");
 
 				// Write the description
-				Writer.EnterTag("[PARAM:description]");
-				if (!Utility.IsNullOrWhitespace(FullDescription))
+				if (!Utility.IsNullOrWhitespace(BriefDescription) || !Utility.IsNullOrWhitespace(FullDescription))
 				{
 					Writer.EnterSection("description", "Remarks");
-					Writer.WriteLine(FullDescription);
+					if (!Utility.IsNullOrWhitespace(BriefDescription) && BriefDescription != FullDescription)
+					{
+						Writer.WriteLine(BriefDescription);
+					}
+					if (!Utility.IsNullOrWhitespace(FullDescription))
+					{
+						Writer.WriteLine(FullDescription);
+					}
 					Writer.LeaveSection();
 				}
-				Writer.LeaveTag("[/PARAM]");
-
-				// Leave the object tag
-				Writer.LeaveTag("[/OBJECT]");
 			}
         }
 
diff --git a/Engine/Source/Programs/UnrealDocTool/APIDocTool/APIDocTool/Metadata.cs b/Engine/Source/Programs/UnrealDocTool/APIDocTool/APIDocTool/Metadata.cs
--- a/Engine/Source/Programs/UnrealDocTool/APIDocTool/APIDocTool/Metadata.cs
+++ b/Engine/Source/Programs/UnrealDocTool/APIDocTool/APIDocTool/Metadata.cs
@@ -226,9 +226,16 @@ namespace APIDocTool
 		public string Description;
 	}
 
+	public class MetadataKeyword
+	{
+		public string Url;
+		public Dictionary<string, string> NodeUrls = new Dictionary<string,string>();
+	}
+
 	public class MetadataDirective
 	{
 		public static Dictionary<string, Dictionary<string, int>> AllTags = new Dictionary<string, Dictionary<string, int>>();
+		public static Dictionary<string, MetadataKeyword> AllKeywords = new Dictionary<string, MetadataKeyword>();
 
 		public readonly SourceFile File;
 		public readonly string Type;
@@ -265,9 +272,55 @@ namespace APIDocTool
 			}
 		}
 
-		public string ToMarkdown()
+		public IEnumerable<string> ToMarkdown()
 		{
-			return Utility.EscapeText(ToString());
+			MetadataKeyword Keyword;
+			if(AllKeywords.TryGetValue(Type, out Keyword))
+			{
+				StringBuilder Output = new StringBuilder();
+				Output.AppendFormat("[{0}]({1})", Type, Keyword.Url);
+				Output.Append(Utility.EscapeText("("));
+
+				int Indent = Type.Length + 1;
+				int HangingIndent = Indent;
+
+				List<string> Lines = new List<string>();
+				for(int Idx = 0; Idx < NodeList.Count; Idx++)
+				{
+					string NodeText = NodeList[Idx].ToString();
+
+					if (Indent + NodeText.Length + 3 > 100)
+					{
+						yield return Output.ToString();
+						Output.Clear();
+						Output.Append(String.Concat(Enumerable.Repeat("&nbsp;", HangingIndent)));
+						Indent = HangingIndent;
+					}
+
+					string NodeUrl;
+					if(Keyword.NodeUrls.TryGetValue(NodeList[Idx].TagText, out NodeUrl))
+					{
+						Output.AppendFormat("[{0}]({1})", Utility.EscapeText(NodeText), NodeUrl);
+					}
+					else
+					{
+						Output.Append(Utility.EscapeText(NodeText));
+					}
+					Indent += NodeText.Length;
+
+					if(Idx + 1 < NodeList.Count)
+					{
+						Output.Append(", ");
+						Indent += 2;
+					}
+				}
+				Output.Append(Utility.EscapeText(")"));
+				yield return Output.ToString();
+			}
+			else
+			{
+				yield return Utility.EscapeText(ToString());
+			}
 		}
 
 		public IEnumerable<Icon> Icons
diff --git a/Engine/Source/Programs/UnrealDocTool/APIDocTool/APIDocTool/Program.cs b/Engine/Source/Programs/UnrealDocTool/APIDocTool/APIDocTool/Program.cs
--- a/Engine/Source/Programs/UnrealDocTool/APIDocTool/APIDocTool/Program.cs
+++ b/Engine/Source/Programs/UnrealDocTool/APIDocTool/APIDocTool/Program.cs
@@ -90,6 +90,7 @@ namespace APIDocTool
 			"NO_API= ",
 			"OVERRIDE= ",
 			"CDECL= ",
+			"DEPRECATED(X,Y)= ",
 
 			// Platform
 			"PLATFORM_SUPPORTS_DRAW_MESH_EVENTS=1",
@@ -378,6 +379,12 @@ namespace APIDocTool
 				IgnoredFunctionMacros = new HashSet<string>(Settings.FindValueOrDefault("Input.IgnoredFunctionMacros", "").Split('\n'));
 				IncludedSourceFiles = new HashSet<string>(Settings.FindValueOrDefault("Output.IncludedSourceFiles", "").Split('\n'));
 
+				// Find all the metadata pages
+				AddMetadataKeyword(UdnDir, "UCLASS", "Programming/UnrealArchitecture/Reference/Classes#classdeclaration", "Programming/UnrealArchitecture/Reference/Classes/Specifiers");
+				AddMetadataKeyword(UdnDir, "UFUNCTION", "Programming/UnrealArchitecture/Reference/Functions", "Programming/UnrealArchitecture/Reference/Functions/Specifiers");
+				AddMetadataKeyword(UdnDir, "UPROPERTY", "Programming/UnrealArchitecture/Reference/Properties", "Programming/UnrealArchitecture/Reference/Properties/Specifiers");
+				AddMetadataKeyword(UdnDir, "USTRUCT", "Programming/UnrealArchitecture/Reference/Structs", "Programming/UnrealArchitecture/Reference/Structs/Specifiers");
+
 				// Clean the output folders
 				if (bCleanTargetInfo)
 				{
@@ -447,6 +454,17 @@ namespace APIDocTool
 			}
 		}
 
+		static void AddMetadataKeyword(string BaseDir, string Name, string TypeUrl, string SpecifierBaseUrl)
+		{
+			MetadataKeyword Keyword = new MetadataKeyword();
+			Keyword.Url = TypeUrl;
+			foreach (DirectoryInfo Specifier in new DirectoryInfo(Path.Combine(BaseDir, SpecifierBaseUrl.Replace('/', '\\'))).EnumerateDirectories())
+			{
+				Keyword.NodeUrls.Add(Specifier.Name, SpecifierBaseUrl.TrimEnd('/') + "/" + Specifier.Name);
+			}
+			MetadataDirective.AllKeywords.Add(Name, Keyword);
+		}
+
 		static void CleanTargetInfo(string TargetInfoPath)
 		{
 			Console.WriteLine("Cleaning '{0}'", TargetInfoPath);
diff --git a/Engine/Source/Programs/UnrealDocTool/APIDocTool/APIDocTool/Utility.cs b/Engine/Source/Programs/UnrealDocTool/APIDocTool/APIDocTool/Utility.cs
--- a/Engine/Source/Programs/UnrealDocTool/APIDocTool/APIDocTool/Utility.cs
+++ b/Engine/Source/Programs/UnrealDocTool/APIDocTool/APIDocTool/Utility.cs
@@ -83,7 +83,8 @@ namespace APIDocTool
 				{
 					foreach (string SubDir in Directory.EnumerateDirectories(Dir))
 					{
-						Directory.Delete(SubDir, true);
+						SafeDeleteDirectoryContents(SubDir, true);
+//						Directory.Delete(SubDir, true);
 					}
 				}
 			}
diff --git a/Engine/Source/Programs/UnrealDocTool/DoxygenLib/Markdown.cs b/Engine/Source/Programs/UnrealDocTool/DoxygenLib/Markdown.cs
--- a/Engine/Source/Programs/UnrealDocTool/DoxygenLib/Markdown.cs
+++ b/Engine/Source/Programs/UnrealDocTool/DoxygenLib/Markdown.cs
@@ -11,7 +11,7 @@ namespace DoxygenLib
 {
 	public static class Markdown
 	{
-		const bool bDocToolCanParseExternalLinks = false;
+		const bool bDocToolCanParseExternalLinks = true;
 
 		public delegate string ResolveLinkDelegate(string Id);
 
@@ -57,7 +57,22 @@ namespace DoxygenLib
 					case "ulink":
 						if (bDocToolCanParseExternalLinks)
 						{
-							output += Indent + String.Format("[{0}]({1})", ParseChildren(child, Indent, ResolveLink), child.Attributes.GetNamedItem("url").InnerText);
+							string LinkText = ParseChildren(child, Indent, ResolveLink);
+							string LinkUrl = child.Attributes.GetNamedItem("url").InnerText;
+
+							const string WebPrefix = "docs.unrealengine.com/latest/INT/";
+							int WebPrefixIdx = LinkUrl.IndexOf(WebPrefix);
+							if(WebPrefixIdx != -1 && LinkText.Contains("docs.unrealengine.com"))
+							{
+								LinkText = "";
+								LinkUrl = LinkUrl.Substring(WebPrefixIdx + WebPrefix.Length);
+								if(LinkUrl.EndsWith(".html"))
+								{
+									LinkUrl = LinkUrl.Substring(0, LinkUrl.LastIndexOf('/') + 1);
+								}
+							}
+
+							output += Indent + String.Format("[{0}]({1})", LinkText, LinkUrl);
 						}
 						else
 						{
@@ -103,6 +118,11 @@ namespace DoxygenLib
 					case "parameternamelist":
 						break;
 					case "simplesect":
+						XmlAttribute KindAttribute = child.Attributes["kind"];
+						if (KindAttribute != null && KindAttribute.Value == "see")
+						{
+							output += "  \n  \nFor more information, see " + ParseChildren(child, Indent + "\t", ResolveLink) + ".";
+						}
 						break;
 					case "xrefsect":
 						break;
@@ -247,5 +267,100 @@ namespace DoxygenLib
 			}
 			return Result.ToString();
 		}
+
+		/** 
+		 * Truncates a line to a given width, respecting embedded links.
+         * 
+		 * @param Text			Text to truncate
+         * @param MaxLength     Maximum number of characters in the output
+         * @param Suffix        Suffix to append if the string is truncated
+         * @return				Truncated string
+		 */
+        public static string Truncate(string Text, int MaxLength, string Suffix)
+        {
+            StringBuilder Output = new StringBuilder();
+            
+            int RemainingCharacters = MaxLength;
+            for (int Idx = 0; Idx < Text.Length; )
+            {
+                int NextIdx = TruncateEscapedCharacter(Text, Idx, ref RemainingCharacters, Output);
+                if(NextIdx == -1)
+                {
+                    NextIdx = TruncateLink(Text, Idx, ref RemainingCharacters, Output);
+                    if(NextIdx == -1)
+                    {
+                        NextIdx = Idx + 1;
+                        if(RemainingCharacters > 0)
+                        {
+                            Output.Append(Text[Idx]);
+                        }
+                        RemainingCharacters--;
+                    }
+                }
+                Idx = NextIdx;
+            }
+            if(RemainingCharacters < 0)
+            {
+                Output.Append(Suffix);
+            }
+
+            return Output.ToString();
+        }
+
+        private static int TruncateEscapedCharacter(string Text, int Idx, ref int RemainingCharacters, StringBuilder Output)
+        {
+            // Check it starts with an ampersand
+            if (Text[Idx] == '&')
+            {
+                // Skip past any sequence of letters
+                int NewIdx = Idx + 1;
+                while (NewIdx < Text.Length && Char.IsLetter(Text[NewIdx]))
+                {
+                    NewIdx++;
+                }
+
+                // Check the whole thing is valid
+                if (NewIdx != Idx + 1 && NewIdx < Text.Length && Text[NewIdx] == ';')
+                {
+                    // Append the markdown
+                    if (RemainingCharacters > 0)
+                    {
+                        Output.Append(Text.Substring(Idx, NewIdx + 1 - Idx));
+                    }
+
+                    // Move to the next character
+                    RemainingCharacters--;
+                    return NewIdx + 1;
+                }
+            }
+            return -1;
+        }
+
+        private static int TruncateLink(string Text, int Idx, ref int RemainingCharacters, StringBuilder Output)
+        {
+            // Check if starts with a bracket
+            if(Text[Idx] == '[')
+            {
+                // Find the end of the description
+                int LabelStart = Idx + 1;
+                int LabelEnd = Text.IndexOf(']', LabelStart);
+                if(LabelEnd != -1 && LabelEnd + 1 < Text.Length && Text[LabelEnd + 1] == '(')
+                {
+                    // Find the end of the link
+                    int LinkStart = LabelEnd + 2;
+                    int LinkEnd = Text.IndexOf(')', LinkStart);
+                    if(LinkEnd != -1)
+                    {
+                        if(RemainingCharacters > 0)
+                        {
+                            Output.AppendFormat("[{0}]({1})", Text.Substring(LabelStart, Math.Min(LabelEnd - LabelStart, RemainingCharacters)), Text.Substring(LinkStart, LinkEnd - LinkStart));
+                        }
+                        RemainingCharacters -= LabelEnd - LabelStart;
+                        return LinkEnd + 1;
+                    }
+                }
+            }
+            return -1;
+        }
 	}
 }
