Scrape comments from enclosing namespaces for enums that are not commented.

diff --git a/Engine/Source/Programs/UnrealDocTool/APIDocTool/APIDocTool/APIEnum.cs b/Engine/Source/Programs/UnrealDocTool/APIDocTool/APIDocTool/APIEnum.cs
--- a/Engine/Source/Programs/UnrealDocTool/APIDocTool/APIDocTool/APIEnum.cs
+++ b/Engine/Source/Programs/UnrealDocTool/APIDocTool/APIDocTool/APIEnum.cs
@@ -57,6 +57,12 @@ namespace APIDocTool
 			// Get the description
 			ParseBriefAndFullDescription(Node, out BriefDescription, out FullDescription);
 
+			// If there isn't one and we're in a namespace, use that instead
+			if(String.IsNullOrEmpty(BriefDescription) && String.IsNullOrEmpty(FullDescription) && Entity.NamespaceNode != null)
+			{
+				ParseBriefAndFullDescription(Entity.NamespaceNode, out BriefDescription, out FullDescription);
+			}
+
 			// Link all the values
 			foreach (APIEnumValue Value in Values)
 			{
@@ -108,7 +114,7 @@ namespace APIDocTool
 					{
 						Writer.WriteLine(BriefDescription);
 					}
-					if (!Utility.IsNullOrWhitespace(FullDescription) && FullDescription != BriefDescription)
+					if (!Utility.IsNullOrWhitespace(FullDescription))
 					{
 						Writer.WriteLine(FullDescription);
 					}
