Remove stray <span> test code, and prevent records repeating their description twice.

diff --git a/Engine/Source/Programs/UnrealDocTool/APIDocTool/APIDocTool/APIRecord.cs b/Engine/Source/Programs/UnrealDocTool/APIDocTool/APIDocTool/APIRecord.cs
--- a/Engine/Source/Programs/UnrealDocTool/APIDocTool/APIDocTool/APIRecord.cs
+++ b/Engine/Source/Programs/UnrealDocTool/APIDocTool/APIDocTool/APIRecord.cs
@@ -283,7 +283,7 @@ namespace APIDocTool
 			}
 			if(BaseDefinitions.Count == 0)
 			{
-				Lines.Add("<span style=\"color:blue\">" + Definition + "</span>");
+				Lines.Add(Definition);
 			}
 			else if(BaseDefinitions.Count == 1)
 			{
@@ -363,7 +363,7 @@ namespace APIDocTool
 				if (!Utility.IsNullOrWhitespace(BriefDescription) || !Utility.IsNullOrWhitespace(FullDescription))
 				{
 					Writer.EnterSection("description", "Remarks");
-					if (!Utility.IsNullOrWhitespace(BriefDescription))
+					if (!Utility.IsNullOrWhitespace(BriefDescription) && BriefDescription != FullDescription)
 					{
 						Writer.WriteLine(BriefDescription);
 					}
