Remove nested classes and typedefs from the class list page.

diff --git a/Engine/Source/Programs/UnrealDocTool/APIDocTool/APIDocTool/APIRecordIndex.cs b/Engine/Source/Programs/UnrealDocTool/APIDocTool/APIDocTool/APIRecordIndex.cs
--- a/Engine/Source/Programs/UnrealDocTool/APIDocTool/APIDocTool/APIRecordIndex.cs
+++ b/Engine/Source/Programs/UnrealDocTool/APIDocTool/APIDocTool/APIRecordIndex.cs
@@ -18,12 +18,15 @@ namespace APIDocTool
 			{
 				if ((Member is APIRecord && (!((APIRecord)Member).bIsTemplateSpecialization)) || (Member is APITypeDef) || (Member is APIEventParameters))
 				{
-					Entry NewEntry = new Entry(Member);
-					if (NewEntry.SortName.Length >= 2 && IgnorePrefixLetters.Contains(NewEntry.SortName[0]) && Char.IsUpper(NewEntry.SortName[1]))
+					if (Member.FindParent<APIRecord>() == null)
 					{
-						NewEntry.SortName = NewEntry.SortName.Substring(1);
+						Entry NewEntry = new Entry(Member);
+						if (NewEntry.SortName.Length >= 2 && IgnorePrefixLetters.Contains(NewEntry.SortName[0]) && Char.IsUpper(NewEntry.SortName[1]))
+						{
+							NewEntry.SortName = NewEntry.SortName.Substring(1);
+						}
+						AddToDefaultCategory(NewEntry);
 					}
-					AddToDefaultCategory(NewEntry);
 				}
 			}
 		}
