Set MarkdownSharpTests to build for x86 instead of Any CPU, to match other projects in the solution.

diff --git a/Engine/Source/Programs/UnrealDocTool/UnrealDocTool/MarkdownSharpTests/MarkdownSharpTests.csproj b/Engine/Source/Programs/UnrealDocTool/UnrealDocTool/MarkdownSharpTests/MarkdownSharpTests.csproj
--- a/Engine/Source/Programs/UnrealDocTool/UnrealDocTool/MarkdownSharpTests/MarkdownSharpTests.csproj
+++ b/Engine/Source/Programs/UnrealDocTool/UnrealDocTool/MarkdownSharpTests/MarkdownSharpTests.csproj
@@ -34,6 +34,7 @@
     <ErrorReport>prompt</ErrorReport>
     <WarningLevel>4</WarningLevel>
     <TreatWarningsAsErrors>true</TreatWarningsAsErrors>
+    <PlatformTarget>x86</PlatformTarget>
   </PropertyGroup>
   <ItemGroup>
     <Reference Include="DotLiquid">
