Switch APIDocTool to use XGE output to determine build environment, instead of relying on special command line switches.

diff --git a/Engine/Source/Programs/UnrealDocTool/APIDocTool/APIDocTool/BuildTarget.cs b/Engine/Source/Programs/UnrealDocTool/APIDocTool/APIDocTool/BuildTarget.cs
--- a/Engine/Source/Programs/UnrealDocTool/APIDocTool/APIDocTool/BuildTarget.cs
+++ b/Engine/Source/Programs/UnrealDocTool/APIDocTool/APIDocTool/BuildTarget.cs
@@ -2,6 +2,7 @@
 
 using System;
 using System.Collections.Generic;
+using System.IO;
 using System.Linq;
 using System.Text;
 using System.Threading.Tasks;
@@ -9,154 +10,118 @@ using System.Xml;
 
 namespace APIDocTool
 {
-	public class BuildModule
+	class BuildTarget
 	{
-		public string Name;
-		public string Path;
-		public string Type;
-		public List<string> Dependencies = new List<string>();
 		public List<string> Definitions = new List<string>();
 		public List<string> IncludePaths = new List<string>();
+		public List<string> DirectoryNames = new List<string>();
 
-		public BuildModule(string InName, string InPath, string InType)
+		public BuildTarget(string BaseDirectoryName, string InputFileName)
 		{
-			Name = InName;
-			Path = InPath;
-			Type = InType;
-		}
-
-		public override string ToString()
-		{
-			return Name;
-		}
+			XmlDocument Document = new XmlDocument();
+			Document.Load(InputFileName);
 
-		public static BuildModule FromXml(XmlNode Node)
-		{
-			BuildModule NewModule = new BuildModule(Node.Attributes["name"].Value, Node.Attributes["path"].Value, Node.Attributes["type"].Value);
-			if (NewModule.Type == "cpp")
+			foreach (XmlNode Node in Document.SelectNodes("BuildSet/Environments/Environment/Tools/Tool"))
 			{
-				using (XmlNodeList DependencyNodeList = Node.SelectNodes("dependencies/dependency"))
+				XmlAttribute ToolPathAttr = Node.Attributes["Path"];
+				if (ToolPathAttr != null && String.Compare(Path.GetFileName(ToolPathAttr.InnerText), "cl.exe") == 0)
 				{
-					foreach (XmlNode DependencyNode in DependencyNodeList)
-					{
-						NewModule.Dependencies.Add(DependencyNode.Attributes["module"].Value);
-					}
-				}
-				using (XmlNodeList DefinitionNodeList = Node.SelectNodes("definitions/definition"))
-				{
-					foreach (XmlNode DefinitionNode in DefinitionNodeList)
+					string Name = Node.Attributes["Name"].InnerText;
+					string Params = Node.Attributes["Params"].InnerText;
+
+					// Parse a list of tokens
+					List<string> Tokens = new List<string>();
+					for (int Idx = 0; Idx < Params.Length; Idx++)
 					{
-						NewModule.Definitions.Add(DefinitionNode.Attributes["name"].Value);
+						if (!Char.IsWhiteSpace(Params[Idx]))
+						{
+							StringBuilder Token = new StringBuilder();
+							while (Idx < Params.Length && !Char.IsWhiteSpace(Params[Idx]))
+							{
+								if (Params[Idx] == '"')
+								{
+									Idx++;
+									while (Params[Idx] != '"') Token.Append(Params[Idx++]);
+									Idx++;
+								}
+								else
+								{
+									Token.Append(Params[Idx++]);
+								}
+							}
+							Tokens.Add(Token.ToString());
+						}
 					}
-				}
-				using (XmlNodeList IncludeNodeList = Node.SelectNodes("includes/include"))
-				{
-					foreach (XmlNode IncludeNode in IncludeNodeList)
+
+					// Parse it into a list of options, removing any that don't apply
+					List<string> FileNames = new List<string>();
+					for (int Idx = 0; Idx < Tokens.Count; Idx++)
 					{
-						NewModule.IncludePaths.Add(IncludeNode.Attributes["path"].Value);
+						if (Tokens[Idx].StartsWith("/Fo") || Tokens[Idx].StartsWith("/Fp") || Tokens[Idx].StartsWith("/Yu"))
+						{
+							Idx++;
+						}
+						else if (Tokens[Idx] == "/D")
+						{
+							Definitions.AddUnique(Tokens[++Idx]);
+						}
+						else if (Tokens[Idx] == "/I")
+						{
+							IncludePaths.AddUnique(ResolveDirectory(BaseDirectoryName, Tokens[++Idx]));
+						}
+						else if (!Tokens[Idx].StartsWith("/"))
+						{
+							DirectoryNames.AddUnique(Path.GetDirectoryName(ResolveFile(BaseDirectoryName, Tokens[Idx])));
+						}
 					}
 				}
 			}
-			return NewModule;
-		}
-	}
-
-	public class BuildBinary
-	{
-		public string Name;
-		public string Type;
-		public List<BuildModule> Modules = new List<BuildModule>();
-
-		public BuildBinary(string InName, string InType)
-		{
-			Name = InName;
-			Type = InType;
-		}
-
-		public BuildBinary Filter(List<string> ModuleNames)
-		{
-			BuildBinary NewBinary = new BuildBinary(Name, Type);
-			foreach (BuildModule Module in Modules)
-			{
-				if (ModuleNames.Contains(Module.Name))
-				{
-					NewBinary.Modules.Add(Module);
-				}
-			}
-			return NewBinary;
-		}
-
-		public override string ToString()
-		{
-			return Name;
 		}
 
-		public static BuildBinary FromXml(XmlNode Node)
+		static string ResolveFile(string BaseDirectoryName, string FileName)
 		{
-			BuildBinary NewBinary = new BuildBinary(Node.Attributes["name"].Value, Node.Attributes["type"].Value);
-			if (NewBinary.Type == "cpp")
+			string NewFileName = FileName;
+			if(!Path.IsPathRooted(NewFileName))
 			{
-				using (XmlNodeList ModuleNodeList = Node.SelectNodes("module"))
-				{
-					foreach (XmlNode ModuleNode in ModuleNodeList)
-					{
-						BuildModule NewModule = BuildModule.FromXml(ModuleNode);
-						NewBinary.Modules.Add(NewModule);
-					}
-				}
+				NewFileName = Path.Combine(BaseDirectoryName, NewFileName);
 			}
-			return NewBinary;
+			return new FileInfo(NewFileName).FullName;
 		}
-	}
-
-	public class BuildTarget
-	{
-		public string Name;
-		public List<BuildBinary> Binaries = new List<BuildBinary>();
 
-		public BuildTarget(string InName)
+		static string ResolveDirectory(string BaseDirectoryName, string DirectoryName)
 		{
-			Name = InName;
-		}
-
-		public BuildTarget Filter(List<string> ModuleNames)
-		{
-			BuildTarget NewTarget = new BuildTarget(Name);
-			foreach (BuildBinary Binary in Binaries)
+			string NewDirectoryName = ExpandEnvironmentVariables(DirectoryName).TrimEnd('/', '\\');
+			if(!Path.IsPathRooted(NewDirectoryName))
 			{
-				BuildBinary NewBinary = Binary.Filter(ModuleNames);
-				if (NewBinary.Modules.Count > 0) NewTarget.Binaries.Add(NewBinary);
+				NewDirectoryName = Path.Combine(BaseDirectoryName, NewDirectoryName);
 			}
-			return NewTarget;
-		}
-
-		public override string ToString()
-		{
-			return Name;
+			return new DirectoryInfo(NewDirectoryName).FullName;
 		}
 
-		public static BuildTarget FromXml(XmlNode Node)
+		static string ExpandEnvironmentVariables(string Text)
 		{
-			BuildTarget NewTarget = new BuildTarget(Node.Attributes["name"].Value);
-			using (XmlNodeList BinaryNodeList = Node.SelectNodes("binary"))
+			int StartIdx = -1;
+			for (int Idx = 0; Idx < Text.Length; Idx++)
 			{
-				foreach (XmlNode BinaryNode in BinaryNodeList)
+				if (Text[Idx] == '$' && (Idx + 1 < Text.Length && Text[Idx + 1] == '('))
+				{
+					// Save the start of a variable name
+					StartIdx = Idx;
+				}
+				else if(Text[Idx] == ')' && StartIdx != -1)
 				{
-					BuildBinary NewBinary = BuildBinary.FromXml(BinaryNode);
-					NewTarget.Binaries.Add(NewBinary);
+					// Replace the variable
+					string Name = Text.Substring(StartIdx + 2, Idx - (StartIdx + 2));
+					string Value = Environment.GetEnvironmentVariable(Name);
+					if (Value != null)
+					{
+						Text = Text.Substring(0, StartIdx) + Value + Text.Substring(Idx + 1);
+						Idx = StartIdx + Value.Length - 1;
+					}
+					StartIdx = -1;
 				}
 			}
-			return NewTarget;
-		}
-
-		public static BuildTarget Read(string InputPath)
-		{
-			// Parse the xml document
-			XmlDocument TargetDocument = Utility.ReadXmlDocument(InputPath);
-
-			// Read the target definition
-			XmlNode TargetNode = TargetDocument.SelectSingleNode("target");
-			return BuildTarget.FromXml(TargetNode);
+			return Text;
 		}
 	}
 }
diff --git a/Engine/Source/Programs/UnrealDocTool/APIDocTool/APIDocTool/Program.cs b/Engine/Source/Programs/UnrealDocTool/APIDocTool/APIDocTool/Program.cs
--- a/Engine/Source/Programs/UnrealDocTool/APIDocTool/APIDocTool/Program.cs
+++ b/Engine/Source/Programs/UnrealDocTool/APIDocTool/APIDocTool/Program.cs
@@ -162,12 +162,6 @@ namespace APIDocTool
 		const string SitemapContentsFileName = "API.hhc";
 		const string SitemapIndexFileName = "API.hhk";
 
-		const string TargetInfoName = "DocumentationEditor";
-		const string TargetInfoPlatform = "Win64";
-		const string TargetInfoConfig = "Debug";
-
-		const string TargetInfoProject = "Documentation\\Extras\\API\\Build\\Documentation.uproject";
-
 		static Program()
 		{
 			List<string> DelegateMacros = new List<string>();
@@ -223,7 +217,6 @@ namespace APIDocTool
 			string XmlDir = Path.Combine(IntermediateRootDir, "doxygen");
 			string StatsPath = null;
 
-			bool bVerbose = false;
 			List<string> Filters = new List<string>();
 
 			// Parse the command line
@@ -334,10 +327,6 @@ namespace APIDocTool
 				{
 					StatsPath = Path.GetFullPath(OptionValue);
 				}
-				else if (OptionName == "-verbose")
-				{
-					bVerbose = true;
-				}
 				else if (OptionName == "-indexonly")
 				{
 					bIndexOnly = true;
@@ -416,9 +405,9 @@ namespace APIDocTool
 				{
 					if (!bBuildMetadata || BuildMetadata(DoxygenPath, EngineDir, MetadataDir, MetadataPath))
 					{
-						if (!bBuildXml || BuildXml(EngineDir, TargetInfoPath, DoxygenPath, XmlDir, Filters, bVerbose))
+						if (!bBuildXml || BuildXml(EngineDir, TargetInfoPath, DoxygenPath, XmlDir, Filters))
 						{
-							if (!bBuildUdn || BuildUdn(XmlDir, UdnDir, ChmDir, MetadataPath, StatsPath, Filters))
+							if (!bBuildUdn || BuildUdn(EngineDir, XmlDir, UdnDir, ChmDir, MetadataPath, StatsPath, Filters))
 							{
 								if (!bBuildHtml || BuildHtml(EngineDir, DocToolPath, UdnDir, HtmlDir))
 								{
@@ -474,28 +463,47 @@ namespace APIDocTool
 		static bool BuildTargetInfo(string TargetInfoPath, string EngineDir)
 		{
 			Console.WriteLine("Building target info...");
-
 			Utility.SafeCreateDirectory(Path.GetDirectoryName(TargetInfoPath));
 
-			using (Process BuildToolProcess = new Process())
+			string Arguments = String.Format("DocumentationEditor Win64 Debug -project=\"{0}\"", Path.Combine(EngineDir, "Documentation\\Extras\\API\\Build\\Documentation.uproject"));
+			if(!RunUnrealBuildTool(EngineDir, Arguments + " -clean"))
+			{
+				return false;
+			}
+			foreach(FileInfo Info in new DirectoryInfo(Path.Combine(EngineDir, "Intermediate\\Build")).EnumerateFiles("UBTEXport*.xml"))
+			{
+				File.Delete(Info.FullName);
+			}
+			if(!RunUnrealBuildTool(EngineDir, Arguments + " -disableunity -xgeexport"))
+			{
+				return false;
+			}
+			File.Copy(Path.Combine(EngineDir, "Intermediate\\Build\\UBTExport.0.xge.xml"), TargetInfoPath, true);
+			return true;
+		}
+
+		static bool RunUnrealBuildTool(string EngineDir, string Arguments)
+		{
+			using (Process NewProcess = new Process())
 			{
-				BuildToolProcess.StartInfo.WorkingDirectory = EngineDir;
-				BuildToolProcess.StartInfo.FileName = Path.Combine(EngineDir, "Binaries\\DotNET\\UnrealBuildTool.exe");
-				BuildToolProcess.StartInfo.Arguments = String.Format("{0} {1} {2} -project=\"{3}\" -writetargetinfo=\"{4}\"", TargetInfoName, TargetInfoPlatform, TargetInfoConfig, Path.Combine(EngineDir, TargetInfoProject), TargetInfoPath);
-				BuildToolProcess.StartInfo.UseShellExecute = false;
-				BuildToolProcess.StartInfo.RedirectStandardOutput = true;
-				BuildToolProcess.StartInfo.RedirectStandardError = true;
+				NewProcess.StartInfo.WorkingDirectory = EngineDir;
+				NewProcess.StartInfo.FileName = Path.Combine(EngineDir, "Binaries\\DotNET\\UnrealBuildTool.exe");
+				NewProcess.StartInfo.Arguments = Arguments;
+				NewProcess.StartInfo.UseShellExecute = false;
+				NewProcess.StartInfo.RedirectStandardOutput = true;
+				NewProcess.StartInfo.RedirectStandardError = true;
+				NewProcess.StartInfo.EnvironmentVariables.Remove("UE_SDKS_ROOT");
 
-				BuildToolProcess.OutputDataReceived += new DataReceivedEventHandler(ProcessOutputReceived);
-				BuildToolProcess.ErrorDataReceived += new DataReceivedEventHandler(ProcessOutputReceived);
+				NewProcess.OutputDataReceived += new DataReceivedEventHandler(ProcessOutputReceived);
+				NewProcess.ErrorDataReceived += new DataReceivedEventHandler(ProcessOutputReceived);
 
 				try
 				{
-					BuildToolProcess.Start();
-					BuildToolProcess.BeginOutputReadLine();
-					BuildToolProcess.BeginErrorReadLine();
-					BuildToolProcess.WaitForExit();
-					return BuildToolProcess.ExitCode == 0;
+					NewProcess.Start();
+					NewProcess.BeginOutputReadLine();
+					NewProcess.BeginErrorReadLine();
+					NewProcess.WaitForExit();
+					return NewProcess.ExitCode == 0;
 				}
 				catch (Exception Ex)
 				{
@@ -505,6 +513,7 @@ namespace APIDocTool
 			}
 		}
 
+
 		static void CleanMetadata(string MetadataDir)
 		{
 			Console.WriteLine("Cleaning '{0}'", MetadataDir);
@@ -516,7 +525,7 @@ namespace APIDocTool
 			string MetadataInputPath = Path.Combine(EngineDir, "Source\\Runtime\\CoreUObject\\Public\\UObject\\ObjectBase.h");
 			Console.WriteLine("Building metadata descriptions from '{0}'...", MetadataInputPath);
 
-			DoxygenConfig Config = new DoxygenConfig("Metadata", MetadataInputPath, MetadataDir);
+			DoxygenConfig Config = new DoxygenConfig("Metadata", new string[]{ MetadataInputPath }, MetadataDir);
 			if (Doxygen.Run(DoxygenPath, Path.Combine(EngineDir, "Source"), Config, true))
 			{
 				MetadataLookup.Reset();
@@ -574,60 +583,49 @@ namespace APIDocTool
 			Utility.SafeDeleteDirectoryContents(XmlDir, true);
 		}
 
-		static bool BuildXml(string EngineDir, string TargetInfoPath, string DoxygenPath, string XmlDir, List<string> Filters = null, bool bVerbose = false)
+		static bool BuildXml(string EngineDir, string TargetInfoPath, string DoxygenPath, string XmlDir, List<string> Filters = null)
 		{
 			// Create the output directory
 			Utility.SafeCreateDirectory(XmlDir);
 
 			// Read the target that we're building
-			BuildTarget Target = BuildTarget.Read(TargetInfoPath);
+			BuildTarget Target = new BuildTarget(Path.Combine(EngineDir, "Source"), TargetInfoPath);
 
 			// Create an invariant list of exclude directories
 			string[] InvariantExcludeDirectories = ExcludeSourceDirectories.Select(x => x.ToLowerInvariant()).ToArray();
 
-			// Flatten the target into a list of modules
-			List<BuildModule> InputModules = new List<BuildModule>();
-			foreach (BuildBinary Binary in Target.Binaries.Where(x => x.Type == "cpp"))
+			// Get the list of folders to filter against
+			List<string> FolderFilters = new List<string>();
+			if(Filters != null)
 			{
-				foreach(BuildModule Module in Binary.Modules)
+				foreach(string Filter in Filters)
 				{
-					if(Module.Type == "cpp")
+					int Idx = Filter.IndexOf('/');
+					if(Idx != -1)
 					{
-						if(!Module.Path.ToLowerInvariant().Split('\\', '/').Any(x => InvariantExcludeDirectories.Contains(x)))
-						{
-							InputModules.Add(Module);
-						}
+						FolderFilters.Add("\\" + Filter.Substring(0, Idx) + "\\");
 					}
 				}
 			}
 
-			// Just use all the input modules
-			List<BuildModule> ParsedModules = new List<BuildModule>();
-			if(bIndexOnly)
-			{
-				ParsedModules.AddRange(InputModules);
-			}
-			else
+			// Flatten the target into a list of modules
+			List<string> InputModules = new List<string>();
+			foreach (string DirectoryName in Target.DirectoryNames)
 			{
-				// Filter the input module list
-				if (Filters != null && Filters.Count > 0)
-				{
-					InputModules = new List<BuildModule>(InputModules.Where(x => Filters.Exists(y => y.StartsWith(x.Name + "/", StringComparison.InvariantCultureIgnoreCase))));
-				}
-
-				// Sort all the modules into the required build order
-				Console.WriteLine("Sorting module parse order...");
-				List<BuildModule> SortedModules = SortModules(InputModules);
-
-				// List all the modules which have circular dependencies
-				foreach (BuildModule Module in SortedModules)
+				for(DirectoryInfo ModuleDirectory = new DirectoryInfo(DirectoryName); ModuleDirectory.Parent != null; ModuleDirectory = ModuleDirectory.Parent)
 				{
-					if (SortedModules.Count(x => x == Module) > 1)
+					IEnumerator<FileInfo> ModuleFile = ModuleDirectory.EnumerateFiles("*.build.cs").GetEnumerator();
+					if(ModuleFile.MoveNext() && (FolderFilters.Count == 0 || FolderFilters.Any(x => ModuleFile.Current.FullName.Contains(x))))
 					{
-						Console.WriteLine("  Circular dependency on {0}; module will be parsed twice", Module.Name);
+						InputModules.AddUnique(ModuleFile.Current.FullName);
+						break;
 					}
 				}
+			}
 
+			// Just use all the input modules
+			if(!bIndexOnly)
+			{
 				// Set our error mode so as to not bring up the WER dialog if Doxygen crashes (which it often does)
 				SetErrorMode(0x0007);
 
@@ -635,109 +633,66 @@ namespace APIDocTool
 				Utility.SafeCreateDirectory(XmlDir);
 
 				// Build all the modules
-				Dictionary<string, string> ModuleToTagfile = new Dictionary<string, string>();
-				for (int Idx = 0; Idx < SortedModules.Count; Idx++)
+				Console.WriteLine("Parsing source...");
+
+				// Build the list of definitions
+				List<string> Definitions = new List<string>();
+				foreach(string Definition in Target.Definitions)
 				{
-					BuildModule Module = SortedModules[Idx];
-					Console.WriteLine("Parsing source for {0}... ({1}/{2})", Module.Name, Idx + 1, SortedModules.Count);
-					string ModuleDir = Path.Combine(XmlDir, Module.Name);
-
-					// Build the configuration for this module
-					DoxygenConfig Config = new DoxygenConfig(Module.Name, Module.Path, ModuleDir);
-					Config.Definitions.AddRange(Module.Definitions.Select(x => x.Replace("DLLIMPORT", "").Replace("DLLEXPORT", "")));
-					Config.Definitions.AddRange(DoxygenPredefinedMacros);
-					Config.ExpandAsDefined.AddRange(DoxygenExpandedMacros);
-					Config.IncludePaths.AddRange(Module.IncludePaths.Select(x => ExpandIncludePath(EngineDir, x).TrimEnd('\\')));
-					Config.ExcludePatterns.AddRange(ExcludeSourceDirectories.Select(x => "*/" + x + "/*"));
-					Config.ExcludePatterns.AddRange(ExcludeSourceFiles);
-					Config.OutputTagfile = Path.Combine(ModuleDir, "tags.xml");
-
-					// Add all the valid tagfiles
-					foreach (string Dependency in Module.Dependencies)
+					if(!Definition.StartsWith("ORIGINAL_FILE_NAME="))
 					{
-						string Tagfile;
-						if (ModuleToTagfile.TryGetValue(Dependency, out Tagfile))
-						{
-							Config.InputTagfiles.Add(Tagfile);
-						}
+						Definitions.Add(Definition.Replace("DLLIMPORT", "").Replace("DLLEXPORT", ""));
 					}
+				}
 
-					// Run doxygen
-					if (Doxygen.Run(DoxygenPath, Path.Combine(EngineDir, "Source"), Config, bVerbose))
+				// Build a list of input paths
+				List<string> InputPaths = new List<string>();
+				foreach(string InputModule in InputModules)
+				{
+					foreach(string DirectoryName in Directory.EnumerateDirectories(Path.GetDirectoryName(InputModule), "*", SearchOption.AllDirectories))
 					{
-						if (!ModuleToTagfile.ContainsKey(Module.Name))
+						// Find the relative path from the engine directory
+						string NormalizedDirectoryName = DirectoryName;
+						if(NormalizedDirectoryName.StartsWith(EngineDir))
 						{
-							ModuleToTagfile.Add(Module.Name, Config.OutputTagfile);
-							ParsedModules.Add(Module);
+							NormalizedDirectoryName = NormalizedDirectoryName.Substring(EngineDir.Length);
+						}
+						if(!NormalizedDirectoryName.EndsWith("\\"))
+						{
+							NormalizedDirectoryName += "\\";
 						}
-					}
-					else
-					{
-						Console.WriteLine("  Doxygen crashed. Skipping.");
-					}
-				}
-			}
-
-			// Write the modules file
-			WriteModulesXml(XmlDir, ParsedModules);
-			return true;
-		}
 
-		static string ExpandEnvironmentVariables(string Text)
-		{
-			int StartIdx = -1;
-			for (int Idx = 0; Idx < Text.Length; Idx++)
-			{
-				if (Text[Idx] == '$' && (Idx + 1 < Text.Length && Text[Idx + 1] == '('))
-				{
-					// Save the start of a variable name
-					StartIdx = Idx;
-				}
-				else if(Text[Idx] == ')' && StartIdx != -1)
-				{
-					// Replace the variable
-					string Name = Text.Substring(StartIdx + 2, Idx - (StartIdx + 2));
-					string Value = Environment.GetEnvironmentVariable(Name);
-					if (Value != null)
-					{
-						Text = Text.Substring(0, StartIdx) + Value + Text.Substring(Idx + 1);
-						Idx = StartIdx + Value.Length - 1;
+						// Check we can include it
+						if(!ExcludeSourceDirectories.Any(x => NormalizedDirectoryName.Contains("\\" + x + "\\")))
+						{
+							if(FolderFilters.Count == 0 || FolderFilters.Any(x => NormalizedDirectoryName.Contains(x)))
+							{
+								InputPaths.Add(DirectoryName);
+							}
+						}
 					}
-					StartIdx = -1;
 				}
-			}
-			return Text;
-		}
-
-		static string ExpandIncludePath(string EngineDir, string IncludePath)
-		{
-			// Expand any environment variables in it
-			string NewIncludePath = ExpandEnvironmentVariables(IncludePath);
 
-			// Convert it to a full path
-			return Path.GetFullPath(Path.Combine(EngineDir, "Source", NewIncludePath));
-		}
-
-		static void WriteModulesXml(string XmlDir, List<BuildModule> ParsedModules)
-		{
-			// Create XML settings for output
-			XmlWriterSettings Settings = new XmlWriterSettings();
-			Settings.Indent = true;
+				// Build the configuration for this module
+				DoxygenConfig Config = new DoxygenConfig("UE4", InputPaths.ToArray(), XmlDir);
+				Config.Definitions.AddRange(Definitions);
+				Config.Definitions.AddRange(DoxygenPredefinedMacros);
+				Config.ExpandAsDefined.AddRange(DoxygenExpandedMacros);
+				Config.IncludePaths.AddRange(Target.IncludePaths);
+				Config.ExcludePatterns.AddRange(ExcludeSourceDirectories.Select(x => "*/" + x + "/*"));
+				Config.ExcludePatterns.AddRange(ExcludeSourceFiles);
 
-			// Save a list of all the converted modules
-			string ModuleXmlPath = Path.Combine(XmlDir, "modules.xml");
-			using(XmlWriter Writer = XmlWriter.Create(ModuleXmlPath, Settings))
-			{
-				Writer.WriteStartElement("modules");
-				foreach(BuildModule Module in ParsedModules)
+				// Run Doxygen
+				if (!Doxygen.Run(DoxygenPath, Path.Combine(EngineDir, "Source"), Config, true))
 				{
-					Writer.WriteStartElement("module");
-					Writer.WriteAttributeString("name", Module.Name);
-					Writer.WriteAttributeString("source", Module.Path);
-					Writer.WriteEndElement();
+					Console.WriteLine("  Doxygen crashed. Skipping.");
+					return false;
 				}
-				Writer.WriteEndElement();
 			}
+
+			// Write the modules file
+			File.WriteAllLines(Path.Combine(XmlDir, "modules.txt"), InputModules);
+			return true;
 		}
 
 		static void CleanUdn(string UdnDir)
@@ -758,7 +713,7 @@ namespace APIDocTool
 			}
 		}
 
-		static bool BuildUdn(string XmlDir, string UdnDir, string SitemapDir, string MetadataPath, string StatsPath, List<string> Filters = null)
+		static bool BuildUdn(string EngineDir, string XmlDir, string UdnDir, string SitemapDir, string MetadataPath, string StatsPath, List<string> Filters = null)
 		{
 			// Create the output directory
 			Utility.SafeCreateDirectory(UdnDir);
@@ -766,43 +721,31 @@ namespace APIDocTool
 			// Read the metadata
 			MetadataLookup.Load(MetadataPath);
 
-			// Read the input module index
-			XmlDocument Document = Utility.ReadXmlDocument(Path.Combine(XmlDir, "modules.xml"));
-
 			// Read the list of modules
-			List<KeyValuePair<string, string>> InputModules = new List<KeyValuePair<string,string>>();
-			using (XmlNodeList NodeList = Document.SelectNodes("modules/module"))
-			{
-				foreach (XmlNode Node in NodeList)
-				{
-					string Name = Node.Attributes["name"].Value;
-					string Source = Node.Attributes["source"].Value;
-					InputModules.Add(new KeyValuePair<string, string>(Name, Source));
-				}
-			}
+			List<string> InputModules = new List<string>(File.ReadAllLines(Path.Combine(XmlDir, "modules.txt")));
 
+			// Build the doxygen modules
 			List<DoxygenModule> Modules = new List<DoxygenModule>();
-			if(bIndexOnly)
+			foreach(string InputModule in InputModules)
 			{
-				// Just create empty modules
-				for (int Idx = 0; Idx < InputModules.Count; Idx++)
-				{
-					Modules.Add(new DoxygenModule(InputModules[Idx].Key, InputModules[Idx].Value));
-				}
+				Modules.Add(new DoxygenModule(Path.GetFileNameWithoutExtension(Path.GetFileNameWithoutExtension(InputModule)), Path.GetDirectoryName(InputModule)));
 			}
-			else
+
+			// Find all the entities
+			if(!bIndexOnly)
 			{
-				// Filter the input module list
-				if (Filters != null && Filters.Count > 0)
+				// Read the engine module and split it into smaller modules
+				DoxygenModule RootModule = DoxygenModule.Read("UE4", EngineDir, Path.Combine(XmlDir, "xml"));
+				foreach(DoxygenEntity Entity in RootModule.Entities)
 				{
-					InputModules = new List<KeyValuePair<string, string>>(InputModules.Where(x => Filters.Exists(y => y.StartsWith(x.Key + "/", StringComparison.InvariantCultureIgnoreCase))));
+					DoxygenModule Module = Modules.Find(x => Entity.File.StartsWith(x.BaseSrcDir));
+					Entity.Module = Module;
+					Module.Entities.Add(Entity);
 				}
-
-				// Read all the doxygen modules
-				for(int Idx = 0; Idx < InputModules.Count; Idx++)
+				foreach(DoxygenSourceFile SourceFile in RootModule.SourceFiles)
 				{
-					Console.WriteLine("Reading module {0}... ({1}/{2})", InputModules[Idx].Key, Idx + 1, InputModules.Count);
-					Modules.Add(DoxygenModule.Read(InputModules[Idx].Key, InputModules[Idx].Value, Path.Combine(XmlDir, InputModules[Idx].Key, "xml")));
+					DoxygenModule Module = Modules.Find(x => SourceFile.FileName.Replace('/', '\\').StartsWith(x.BaseSrcDir));
+					Module.SourceFiles.Add(SourceFile);
 				}
 
 				// Now filter all the entities in each module
@@ -810,6 +753,9 @@ namespace APIDocTool
 				{
 					FilterEntities(Modules, Filters);
 				}
+
+				// Remove all the empty modules
+				Modules.RemoveAll(x => x.Entities.Count == 0);
 			}
 
 			// Create the index page, and all the pages below it
@@ -959,67 +905,6 @@ namespace APIDocTool
 			}
 		}
 
-		public static List<BuildModule> SortModules(IEnumerable<BuildModule> Modules)
-		{
-			List<BuildModule> OutputModules = new List<BuildModule>();
-
-			// Build a set of unbuilt module names
-			HashSet<string> UnbuiltModules = new HashSet<string>();
-			foreach (BuildModule Module in Modules)
-			{
-				UnbuiltModules.Add(Module.Name);
-			}
-
-			// Build the output list by iteratively removing modules whose dependencies are met
-			LinkedList<BuildModule> RemainingModules = new LinkedList<BuildModule>(Modules);
-			while(RemainingModules.Count > 0)
-			{
-				int InitialOutputModuleCount = OutputModules.Count;
-
-				// Move all the modules whose dependencies are met to the output list
-				for(LinkedListNode<BuildModule> ModuleNode = RemainingModules.First; ModuleNode != null; )
-				{
-					LinkedListNode<BuildModule> NextModuleNode = ModuleNode.Next;
-					if(HasAllDependencies(ModuleNode.Value, UnbuiltModules))
-					{
-						OutputModules.Add(ModuleNode.Value);
-						UnbuiltModules.Remove(ModuleNode.Value.Name);
-						RemainingModules.Remove(ModuleNode);
-					}
-					ModuleNode = NextModuleNode;
-				}
-
-				// If we didn't manage to build anything on this iteration, we've got a recursive dependency. Try to break it by builting the first module that hasn't been seen 
-				// yet (but still build it a second time later)
-				if(OutputModules.Count == InitialOutputModuleCount)
-				{
-					for (LinkedListNode<BuildModule> ModuleNode = RemainingModules.First; ModuleNode != null; ModuleNode = ModuleNode.Next)
-					{
-						BuildModule Module = ModuleNode.Value;
-						if(UnbuiltModules.Contains(Module.Name))
-						{
-							OutputModules.Add(Module);
-							UnbuiltModules.Remove(Module.Name);
-							break;
-						}
-					}
-				}
-			}
-			return OutputModules;
-		}
-
-		public static bool HasAllDependencies(BuildModule Module, HashSet<string> UnbuiltModules)
-		{
-			foreach (string Dependency in Module.Dependencies)
-			{
-				if (UnbuiltModules.Contains(Dependency))
-				{
-					return false;
-				}
-			}
-			return true;
-		}
-
 		public static void CleanHtml(string HtmlPath)
 		{
 			string CleanDir = Path.Combine(HtmlPath, "INT\\API");
@@ -1035,7 +920,7 @@ namespace APIDocTool
 			{
 				DocToolProcess.StartInfo.WorkingDirectory = EngineDir;
 				DocToolProcess.StartInfo.FileName = DocToolPath;
-				DocToolProcess.StartInfo.Arguments = "API\\* -lang=INT -t=DefaultAPI.html";
+				DocToolProcess.StartInfo.Arguments = "API\\* -lang=INT -t=DefaultAPI.html -v=warn";
 				DocToolProcess.StartInfo.UseShellExecute = false;
 				DocToolProcess.StartInfo.RedirectStandardOutput = true;
 				DocToolProcess.StartInfo.RedirectStandardError = true;
diff --git a/Engine/Source/Programs/UnrealDocTool/APIDocTool/APIDocTool/Utility.cs b/Engine/Source/Programs/UnrealDocTool/APIDocTool/APIDocTool/Utility.cs
--- a/Engine/Source/Programs/UnrealDocTool/APIDocTool/APIDocTool/Utility.cs
+++ b/Engine/Source/Programs/UnrealDocTool/APIDocTool/APIDocTool/Utility.cs
@@ -16,6 +16,14 @@ namespace APIDocTool
 	{
 		static HashSet<string> AllLinkPaths = new HashSet<string>();
 
+		public static void AddUnique<T>(this List<T> Items, T Item)
+		{
+			if(!Items.Contains(Item))
+			{
+				Items.Add(Item);
+			}
+		}
+
 		public static void AddToDictionaryList<T, U>(T Key, U Value, Dictionary<T, List<U>> Dict)
 		{
 			List<U> ItemList;
diff --git a/Engine/Source/Programs/UnrealDocTool/DoxygenLib/Doxygen.cs b/Engine/Source/Programs/UnrealDocTool/DoxygenLib/Doxygen.cs
--- a/Engine/Source/Programs/UnrealDocTool/DoxygenLib/Doxygen.cs
+++ b/Engine/Source/Programs/UnrealDocTool/DoxygenLib/Doxygen.cs
@@ -13,7 +13,7 @@ namespace DoxygenLib
 	public class DoxygenConfig
 	{
 		public string Name;
-		public string InputPath;
+		public string[] InputPaths;
 		public string OutputPath;
 		public string OutputTagfile;
 
@@ -23,10 +23,10 @@ namespace DoxygenLib
 		public List<string> InputTagfiles = new List<string>();
 		public List<string> ExpandAsDefined = new List<string>();
 
-		public DoxygenConfig(string InName, string InInputPath, string InOutputPath)
+		public DoxygenConfig(string InName, string[] InInputPaths, string InOutputPath)
 		{
 			Name = InName;
-			InputPath = InInputPath;
+			InputPaths = InInputPaths;
 			OutputPath = InOutputPath;
 		}
 
@@ -47,7 +47,7 @@ namespace DoxygenLib
 				FormatSetting(Output, "GENERATE_LATEX", "NO");
 				FormatSetting(Output, "GENERATE_XML", "YES");
 				FormatSetting(Output, "WARN_IF_UNDOCUMENTED", "NO");
-				FormatSetting(Output, "QUIET", "YES");
+//				FormatSetting(Output, "QUIET", "YES");
 				FormatSetting(Output, "EXTRACT_ALL", "YES");
 				FormatSetting(Output, "WARNINGS", "YES");
 				FormatSetting(Output, "RECURSIVE", "NO");
@@ -65,14 +65,6 @@ namespace DoxygenLib
 					FormatSetting(Output, "GENERATE_TAGFILE", QuoteValue(OutputTagfile));
 				}
 
-				// Build the input paths
-				List<string> InputPaths = new List<string>();
-				InputPaths.Add(InputPath);
-				if (Directory.Exists(InputPath))
-				{
-					InputPaths.AddRange(Directory.EnumerateDirectories(InputPath, "*", SearchOption.AllDirectories));
-				}
-
 				// List of input paths
 				Output.WriteLine();
 				FormatSetting(Output, "INPUT", new List<string>(InputPaths.Select(x => QuoteValue(x))));
@@ -169,7 +161,7 @@ namespace DoxygenLib
 			{
 				DoxygenProcess.StartInfo.WorkingDirectory = WorkingDir;
 				DoxygenProcess.StartInfo.FileName = DoxygenPath;
-				DoxygenProcess.StartInfo.Arguments = "\"" + ConfigFilePath + "\"";
+				DoxygenProcess.StartInfo.Arguments = "\"" + ConfigFilePath + "\" -b";
 				DoxygenProcess.StartInfo.UseShellExecute = false;
 				DoxygenProcess.StartInfo.RedirectStandardOutput = true;
 				DoxygenProcess.StartInfo.RedirectStandardError = true;
diff --git a/Engine/Source/Programs/UnrealDocTool/DoxygenLib/DoxygenModule.cs b/Engine/Source/Programs/UnrealDocTool/DoxygenLib/DoxygenModule.cs
--- a/Engine/Source/Programs/UnrealDocTool/DoxygenLib/DoxygenModule.cs
+++ b/Engine/Source/Programs/UnrealDocTool/DoxygenLib/DoxygenModule.cs
@@ -33,12 +33,13 @@ namespace DoxygenLib
 
 	public class DoxygenEntity
 	{
-		public readonly DoxygenModule Module;
+		public DoxygenModule Module;
 		public DoxygenEntity Parent;
 		public string Id;
 		public string Name;
 		public string Kind;
 		public XmlNode Node;
+		public XmlNode NamespaceNode;
 		public string File;
 		public int Line;
 		public int Column;
@@ -55,13 +56,14 @@ namespace DoxygenLib
 			BodyStart = BodyEnd = -1;
 		}
 
-		public static DoxygenEntity FromXml(DoxygenModule InModule, string Name, XmlNode Node)
+		public static DoxygenEntity FromXml(DoxygenModule InModule, string Name, XmlNode Node, XmlNode NamespaceNode)
 		{
 			DoxygenEntity Entity = new DoxygenEntity(InModule);
 			Entity.Id = Node.Attributes["id"].Value;
 			Entity.Name = Name;
 			Entity.Kind = Node.Attributes["kind"].Value;
 			Entity.Node = Node;
+			Entity.NamespaceNode = NamespaceNode;
 
 			XmlNode LocationNode = Entity.Node.SelectSingleNode("location");
 			if (LocationNode != null)
@@ -149,7 +151,6 @@ namespace DoxygenLib
 	{
 		public readonly string Name;
 		public readonly string BaseSrcDir;
-		public List<DoxygenCompound> Compounds = new List<DoxygenCompound>();
 		public List<DoxygenEntity> Entities = new List<DoxygenEntity>();
 		public List<DoxygenSourceFile> SourceFiles = new List<DoxygenSourceFile>();
 
@@ -208,18 +209,18 @@ namespace DoxygenLib
 			{
 				if (Compound.Kind == "file")
 				{
-					ReadMembers(Module, Compound, "", null, Module.Entities);
+					ReadMembers(Module, Compound, "", null, null, Module.Entities);
 					DoxygenSourceFile SourceFile = DoxygenSourceFile.FromXml(Compound.Node);
 					if(SourceFile != null) Module.SourceFiles.Add(SourceFile);
 				}
 				else if (Compound.Kind == "namespace")
 				{
-					ReadMembers(Module, Compound, Compound.Name + "::", null, Module.Entities);
+					ReadMembers(Module, Compound, Compound.Name + "::", null, Compound.Node, Module.Entities);
 				}
 				else if (Compound.Kind == "class" || Compound.Kind == "struct" || Compound.Kind == "union")
 				{
-					DoxygenEntity Entity = DoxygenEntity.FromXml(Module, Compound.Name, Compound.Node);
-					ReadMembers(Module, Compound, "", Entity, Entity.Members);
+					DoxygenEntity Entity = DoxygenEntity.FromXml(Module, Compound.Name, Compound.Node, null);
+					ReadMembers(Module, Compound, "", Entity, null, Entity.Members);
 					Scopes.Add(Entity.Name, Entity);
 				}
 			}
@@ -281,7 +282,7 @@ namespace DoxygenLib
 			}
 		}
 
-		protected static void ReadMembers(DoxygenModule Module, DoxygenCompound Compound, string NamePrefix, DoxygenEntity Parent, List<DoxygenEntity> Entities)
+		protected static void ReadMembers(DoxygenModule Module, DoxygenCompound Compound, string NamePrefix, DoxygenEntity Parent, XmlNode NamespaceNode, List<DoxygenEntity> Entities)
 		{
 			using (XmlNodeList NodeList = Compound.Node.SelectNodes("sectiondef/memberdef"))
 			{
@@ -289,7 +290,7 @@ namespace DoxygenLib
 				{
 					XmlNode NameNode = Node.SelectSingleNode("name");
 					string Name = NamePrefix + NameNode.InnerText;
-					DoxygenEntity Entity = DoxygenEntity.FromXml(Module, Name, Node);
+					DoxygenEntity Entity = DoxygenEntity.FromXml(Module, Name, Node, NamespaceNode);
 					Entity.Parent = Parent;
 					Entities.Add(Entity);
 				}
