Fix brief description being repeated twice on function pages.

diff --git a/Engine/Source/Programs/UnrealDocTool/APIDocTool/APIDocTool/APIFunction.cs b/Engine/Source/Programs/UnrealDocTool/APIDocTool/APIDocTool/APIFunction.cs
--- a/Engine/Source/Programs/UnrealDocTool/APIDocTool/APIDocTool/APIFunction.cs
+++ b/Engine/Source/Programs/UnrealDocTool/APIDocTool/APIDocTool/APIFunction.cs
@@ -691,7 +691,7 @@ namespace APIDocTool
 					}
 					if (!Utility.IsNullOrWhitespace(FullDescription))
 					{
-						Writer.WriteLine(BriefDescription);
+						Writer.WriteLine(FullDescription);
 					}
 					Writer.LeaveSection();
 				}
