APIDocTool changes.
- Use custom brief description parsing rather than relying on Doxygen. Allows formatting to be preserved in long-form descriptions.
- [REDACTED_HANDLE] blocks are separated into sections
- Markdown conversion now handles paragraphs and lists correctly

diff --git a/Engine/Source/Programs/UnrealDocTool/APIDocTool/APIDocTool/APIConstant.cs b/Engine/Source/Programs/UnrealDocTool/APIDocTool/APIDocTool/APIConstant.cs
--- a/Engine/Source/Programs/UnrealDocTool/APIDocTool/APIDocTool/APIConstant.cs
+++ b/Engine/Source/Programs/UnrealDocTool/APIDocTool/APIDocTool/APIConstant.cs
@@ -38,17 +38,10 @@ namespace APIDocTool
 				Writer.LeaveSection();
 
 				// Write the description
-				if (!Utility.IsNullOrWhitespace(BriefDescription) || !Utility.IsNullOrWhitespace(FullDescription))
+				if (!Utility.IsNullOrWhitespace(FullDescription))
 				{
 					Writer.EnterSection("description", "Remarks");
-					if (!Utility.IsNullOrWhitespace(BriefDescription) && FullDescription != BriefDescription)
-					{
-						Writer.WriteLine(BriefDescription);
-					}
-					if(!Utility.IsNullOrWhitespace(FullDescription))
-					{
-						Writer.WriteLine(FullDescription);
-					}
+					Writer.WriteLine(FullDescription);
 					Writer.LeaveSection();
 				}
 
diff --git a/Engine/Source/Programs/UnrealDocTool/APIDocTool/APIDocTool/APIEnum.cs b/Engine/Source/Programs/UnrealDocTool/APIDocTool/APIDocTool/APIEnum.cs
--- a/Engine/Source/Programs/UnrealDocTool/APIDocTool/APIDocTool/APIEnum.cs
+++ b/Engine/Source/Programs/UnrealDocTool/APIDocTool/APIDocTool/APIEnum.cs
@@ -107,17 +107,10 @@ namespace APIDocTool
 				Writer.LeaveSection();
 
 				// Write the enum values
-				if (!Utility.IsNullOrWhitespace(BriefDescription) || !Utility.IsNullOrWhitespace(FullDescription))
+				if (!Utility.IsNullOrWhitespace(FullDescription))
 				{
 					Writer.EnterSection("description", "Remarks");
-					if (!Utility.IsNullOrWhitespace(BriefDescription) && FullDescription != BriefDescription)
-					{
-						Writer.WriteLine(BriefDescription);
-					}
-					if (!Utility.IsNullOrWhitespace(FullDescription))
-					{
-						Writer.WriteLine(FullDescription);
-					}
+					Writer.WriteLine(FullDescription);
 					Writer.LeaveSection();
 				}
 
diff --git a/Engine/Source/Programs/UnrealDocTool/APIDocTool/APIDocTool/APIFunction.cs b/Engine/Source/Programs/UnrealDocTool/APIDocTool/APIDocTool/APIFunction.cs
--- a/Engine/Source/Programs/UnrealDocTool/APIDocTool/APIDocTool/APIFunction.cs
+++ b/Engine/Source/Programs/UnrealDocTool/APIDocTool/APIDocTool/APIFunction.cs
@@ -219,6 +219,11 @@ namespace APIDocTool
 			get { return APIMember.ConvertToMarkdown(Node.SelectSingleNode("type")); }
 		}
 
+		public string TypeText
+		{
+			get { return Node.SelectSingleNode("type").InnerText; }
+		}
+
 		public string GetAbbreviatedDefinition()
 		{
 			string Text = Definition;
@@ -286,6 +291,7 @@ namespace APIDocTool
 		public string BriefDescription = "";
 		public string FullDescription = "";
 		public string ReturnDescription = "";
+		public List<string> SeeAlso = new List<string>();
 
 		public APIEventParameters EventParameters;
 
@@ -393,6 +399,9 @@ namespace APIDocTool
 				}
 			}
 
+			// Get the @see directives
+			ParseSeeAlso(Node, SeeAlso);
+
 			// Get the modifiers
 			IsVirtual = Node.Attributes.GetNamedItem("virt").InnerText == "virtual";
 			IsStatic = Node.Attributes.GetNamedItem("static").InnerText == "yes";
@@ -682,17 +691,10 @@ namespace APIDocTool
 				Writer.LeaveSection();
 
 				// Write the description
-				if (!Utility.IsNullOrWhitespace(BriefDescription) || !Utility.IsNullOrWhitespace(FullDescription))
+				if (!Utility.IsNullOrWhitespace(FullDescription))
 				{
 					Writer.EnterSection("description", "Remarks");
-					if (!Utility.IsNullOrWhitespace(BriefDescription) && BriefDescription != FullDescription)
-					{
-						Writer.WriteLine(BriefDescription);
-					}
-					if (!Utility.IsNullOrWhitespace(FullDescription))
-					{
-						Writer.WriteLine(FullDescription);
-					}
+					Writer.WriteLine(FullDescription);
 					Writer.LeaveSection();
 				}
 
@@ -738,6 +740,9 @@ namespace APIDocTool
 				// Write the source
 				WriteSourceSection(Writer);
 
+				// Write the @see directives
+				WriteSeeAlsoSection(Writer, SeeAlso);
+
 				// Write the reference info
 				WriteReferencesSection(Writer, Entity);
 			}
@@ -807,7 +812,7 @@ namespace APIDocTool
 			if (FunctionArray.Length > 0)
 			{
 				Writer.EnterSection(SectionId, SectionTitle);
-				WriteList(Writer, FunctionArray);
+				WriteList(Writer, Functions);
 				Writer.LeaveSection();
 				return true;
 			}
diff --git a/Engine/Source/Programs/UnrealDocTool/APIDocTool/APIDocTool/APIMember.cs b/Engine/Source/Programs/UnrealDocTool/APIDocTool/APIDocTool/APIMember.cs
--- a/Engine/Source/Programs/UnrealDocTool/APIDocTool/APIDocTool/APIMember.cs
+++ b/Engine/Source/Programs/UnrealDocTool/APIDocTool/APIDocTool/APIMember.cs
@@ -364,6 +364,54 @@ namespace APIDocTool
 			}
 		}
 
+		protected static void ParseSeeAlso(XmlNode Node, List<string> SeeAlso)
+		{
+			// Parse out all the @see nodes
+			foreach (XmlNode SectionNode in Node.SelectNodes("detaileddescription/para/simplesect"))
+			{
+				XmlAttribute KindAttribute = SectionNode.Attributes["kind"];
+				if (KindAttribute != null && KindAttribute.Value == "see")
+				{
+					foreach (XmlNode ChildNode in SectionNode.ChildNodes)
+					{
+						if (ChildNode.Name != "simplesectsep")
+						{
+							SeeAlso.Add(ConvertToMarkdown(ChildNode));
+						}
+					}
+				}
+			}
+
+			// Split them all by commas
+			for (int SeeAlsoIdx = 0; SeeAlsoIdx < SeeAlso.Count; SeeAlsoIdx++)
+			{
+				string Line = SeeAlso[SeeAlsoIdx];
+				for (int Idx = 0; Idx < Line.Length; Idx = Markdown.SkipCharacter(Line, Idx))
+				{
+					if (Line[Idx] == ',')
+					{
+						SeeAlso[SeeAlsoIdx] = Line.Substring(0, Idx).TrimEnd();
+						SeeAlso.Insert(SeeAlsoIdx + 1, Line.Substring(Idx + 1).TrimStart());
+						break;
+					}
+				}
+			}
+		}
+
+		protected static void WriteSeeAlsoSection(UdnWriter Writer, List<string> SeeAlso)
+		{
+			if (SeeAlso.Count > 0)
+			{
+				Writer.EnterSection("seealso", "See Also");
+				foreach (string SeeAlsoLine in SeeAlso)
+				{
+					Writer.WriteLine(SeeAlsoLine);
+					Writer.WriteLine();
+				}
+				Writer.LeaveSection();
+			}
+		}
+
 		public static APIProtection ParseProtection(XmlNode Node)
 		{
 			XmlNode ProtNode = Node.Attributes.GetNamedItem("prot");
diff --git a/Engine/Source/Programs/UnrealDocTool/APIDocTool/APIDocTool/APIRecord.cs b/Engine/Source/Programs/UnrealDocTool/APIDocTool/APIDocTool/APIRecord.cs
--- a/Engine/Source/Programs/UnrealDocTool/APIDocTool/APIDocTool/APIRecord.cs
+++ b/Engine/Source/Programs/UnrealDocTool/APIDocTool/APIDocTool/APIRecord.cs
@@ -39,6 +39,7 @@ namespace APIDocTool
 
 		public string BriefDescription;
 		public string FullDescription;
+		public List<string> SeeAlso = new List<string>();
 
 		public string TemplateSignature;
 		public string Definition;
@@ -151,6 +152,9 @@ namespace APIDocTool
 			// Get the descriptions
 			ParseBriefAndFullDescription(Node, out BriefDescription, out FullDescription);
 
+			// Get the @see directives
+			ParseSeeAlso(Node, SeeAlso);
+
 			// Find the metadata for this definition
 			MetadataDirective = MetadataCache.FindMetadataForMember(Node);
 
@@ -360,17 +364,10 @@ namespace APIDocTool
 				}
 
 				// Write the class description
-				if (!Utility.IsNullOrWhitespace(BriefDescription) || !Utility.IsNullOrWhitespace(FullDescription))
+				if (!Utility.IsNullOrWhitespace(FullDescription))
 				{
 					Writer.EnterSection("description", "Remarks");
-					if (!Utility.IsNullOrWhitespace(BriefDescription) && BriefDescription != FullDescription)
-					{
-						Writer.WriteLine(BriefDescription);
-					}
-					if (!Utility.IsNullOrWhitespace(FullDescription))
-					{
-						Writer.WriteLine(FullDescription);
-					}
+					Writer.WriteLine(FullDescription);
 					Writer.LeaveSection();
 				}
 
@@ -463,6 +460,9 @@ namespace APIDocTool
 					Writer.LeaveSection();
 				}
 
+				// Write the @see directives
+				WriteSeeAlsoSection(Writer, SeeAlso);
+
 				// Write the references
 				WriteReferencesSection(Writer, Entity);
 			}
diff --git a/Engine/Source/Programs/UnrealDocTool/APIDocTool/APIDocTool/APITypeDef.cs b/Engine/Source/Programs/UnrealDocTool/APIDocTool/APIDocTool/APITypeDef.cs
--- a/Engine/Source/Programs/UnrealDocTool/APIDocTool/APIDocTool/APITypeDef.cs
+++ b/Engine/Source/Programs/UnrealDocTool/APIDocTool/APIDocTool/APITypeDef.cs
@@ -22,7 +22,7 @@ namespace APIDocTool
 
 		public APITypeDef(APIPage InParent, DoxygenEntity InEntity) 
 			: base(InParent, InEntity.Node.SelectSingleNode("name").InnerText)
-        {
+		{
 			Entity = InEntity;
 			Node = Entity.Node;
 			Id = Node.Attributes["id"].Value;
@@ -44,7 +44,7 @@ namespace APIDocTool
 
 		public void ReplaceMarkdownVariables(ref string Text)
 		{
-			Text = Text.Replace("&#123;&#123; typedef&#45;type &#125;&#125;", Type);
+			Text = Text.Replace("{{ typedef-type }}", Type);
 		}
 
 		public override void WritePage(UdnManifest Manifest, string Path)
@@ -59,17 +59,10 @@ namespace APIDocTool
 				Writer.LeaveSection();
 
 				// Write the description
-				if (!Utility.IsNullOrWhitespace(BriefDescription) || !Utility.IsNullOrWhitespace(FullDescription))
+				if (!Utility.IsNullOrWhitespace(FullDescription))
 				{
 					Writer.EnterSection("description", "Remarks");
-					if (!Utility.IsNullOrWhitespace(BriefDescription) && BriefDescription != FullDescription)
-					{
-						Writer.WriteLine(BriefDescription);
-					}
-					if (!Utility.IsNullOrWhitespace(FullDescription))
-					{
-						Writer.WriteLine(FullDescription);
-					}
+					Writer.WriteLine(FullDescription);
 					Writer.LeaveSection();
 				}
 
diff --git a/Engine/Source/Programs/UnrealDocTool/APIDocTool/APIDocTool/APIVariable.cs b/Engine/Source/Programs/UnrealDocTool/APIDocTool/APIDocTool/APIVariable.cs
--- a/Engine/Source/Programs/UnrealDocTool/APIDocTool/APIDocTool/APIVariable.cs
+++ b/Engine/Source/Programs/UnrealDocTool/APIDocTool/APIDocTool/APIVariable.cs
@@ -130,17 +130,10 @@ namespace APIDocTool
 				Writer.LeaveSection();
 
 				// Write the description
-				if (!Utility.IsNullOrWhitespace(BriefDescription) || !Utility.IsNullOrWhitespace(FullDescription))
+				if (!Utility.IsNullOrWhitespace(FullDescription))
 				{
 					Writer.EnterSection("description", "Remarks");
-					if (!Utility.IsNullOrWhitespace(BriefDescription) && BriefDescription != FullDescription)
-					{
-						Writer.WriteLine(BriefDescription);
-					}
-					if (!Utility.IsNullOrWhitespace(FullDescription))
-					{
-						Writer.WriteLine(FullDescription);
-					}
+					Writer.WriteLine(FullDescription);
 					Writer.LeaveSection();
 				}
 			}
diff --git a/Engine/Source/Programs/UnrealDocTool/APIDocTool/APIDocTool/Program.cs b/Engine/Source/Programs/UnrealDocTool/APIDocTool/APIDocTool/Program.cs
--- a/Engine/Source/Programs/UnrealDocTool/APIDocTool/APIDocTool/Program.cs
+++ b/Engine/Source/Programs/UnrealDocTool/APIDocTool/APIDocTool/Program.cs
@@ -863,6 +863,8 @@ namespace APIDocTool
 			// Find all the entities
 			if(!bIndexOnly)
 			{
+				Console.WriteLine("Reading Doxygen output...");
+
 				// Read the engine module and split it into smaller modules
 				DoxygenModule RootModule = DoxygenModule.Read("UE4", EngineDir, Path.Combine(XmlDir, "xml"));
 				foreach(DoxygenEntity Entity in RootModule.Entities)
diff --git a/Engine/Source/Programs/UnrealDocTool/APIDocTool/APIDocTool/Utility.cs b/Engine/Source/Programs/UnrealDocTool/APIDocTool/APIDocTool/Utility.cs
--- a/Engine/Source/Programs/UnrealDocTool/APIDocTool/APIDocTool/Utility.cs
+++ b/Engine/Source/Programs/UnrealDocTool/APIDocTool/APIDocTool/Utility.cs
@@ -75,7 +75,13 @@ namespace APIDocTool
 		{
 			if (Directory.Exists(Dir))
 			{
-				Directory.Delete(Dir, true);
+				try
+				{
+					Directory.Delete(Dir, true);
+				}
+				catch(IOException)
+				{
+				}
 			}
 		}
 
@@ -266,14 +272,6 @@ namespace APIDocTool
 
 		public static string GetBriefDescription(string Description)
 		{
-			const int DesiredLength = 200;
-
-			// If the description is already short, just use that.
-			if(Description.Length < DesiredLength)
-			{
-				return Description;
-			}
-
 			// Find the first newline, or period that's followed by a space and not within any type of paren
 			int Nesting = 0;
 			for(int Idx = 0; Idx < Description.Length; Idx++)
diff --git a/Engine/Source/Programs/UnrealDocTool/DoxygenLib/Doxygen.cs b/Engine/Source/Programs/UnrealDocTool/DoxygenLib/Doxygen.cs
--- a/Engine/Source/Programs/UnrealDocTool/DoxygenLib/Doxygen.cs
+++ b/Engine/Source/Programs/UnrealDocTool/DoxygenLib/Doxygen.cs
@@ -56,7 +56,6 @@ namespace DoxygenLib
 				FormatSetting(Output, "WARN_NO_PARAMDOC", "NO");
 				FormatSetting(Output, "ENABLE_PREPROCESSING", "YES");
 				FormatSetting(Output, "MACRO_EXPANSION", "YES");
-				FormatSetting(Output, "JAVADOC_AUTOBRIEF", "YES");
 				FormatSetting(Output, "SKIP_FUNCTION_MACROS", "NO");
 
 				// Write the conditional settings
diff --git a/Engine/Source/Programs/UnrealDocTool/DoxygenLib/Markdown.cs b/Engine/Source/Programs/UnrealDocTool/DoxygenLib/Markdown.cs
--- a/Engine/Source/Programs/UnrealDocTool/DoxygenLib/Markdown.cs
+++ b/Engine/Source/Programs/UnrealDocTool/DoxygenLib/Markdown.cs
@@ -9,152 +9,243 @@ using System.Xml;
 
 namespace DoxygenLib
 {
+	public class MarkdownWriter
+	{
+		StringBuilder Buffer = new StringBuilder();
+		bool bWantNewParagraph;
+		bool bIsNewLine;
+		string Indent;
+		Stack<string> IndentStack = new Stack<string>();
+
+		public MarkdownWriter(string InIndent)
+		{
+			Indent = InIndent;
+		}
+
+		public void AddIndent(string ExtraIndent)
+		{
+			IndentStack.Push(Indent);
+			Indent += ExtraIndent;
+		}
+
+		public void RemoveIndent()
+		{
+			Indent = IndentStack.Pop();
+		}
+
+		public void WriteParagraph()
+		{
+			bWantNewParagraph = true;
+		}
+
+		public void Write(string Text)
+		{
+			FlushParagraph();
+			if(bIsNewLine)
+			{
+				Buffer.Append(Indent);
+				bIsNewLine = false;
+			}
+			Buffer.Append(Text);
+		}
+
+		public void WriteFormat(string Format, params object[] Args)
+		{
+			FlushParagraph();
+			if (bIsNewLine)
+			{
+				Buffer.Append(Indent);
+				bIsNewLine = false;
+			}
+			Buffer.AppendFormat(Format, Args);
+		}
+
+		public void WriteLine()
+		{
+			if (!bIsNewLine)
+			{
+				Buffer.AppendLine();
+				bIsNewLine = true;
+			}
+		}
+
+		public void WriteLink(string LinkText, string LinkUrl)
+		{
+			WriteFormat("[{0}]({1})", LinkText, LinkUrl);
+		}
+
+		void FlushParagraph()
+		{
+			if (bWantNewParagraph)
+			{
+				if (!bIsNewLine)
+				{
+					Buffer.AppendLine();
+					bIsNewLine = true;
+				}
+				Buffer.AppendLine();
+				bWantNewParagraph = false;
+			}
+		}
+
+		public override string ToString()
+		{
+			return Buffer.ToString();
+		}
+	}
+
 	public static class Markdown
 	{
 		const bool bDocToolCanParseExternalLinks = true;
 
 		public delegate string ResolveLinkDelegate(string Id);
 
-		/** Parses markdown text out of an XML node 
+		/** 
+		 * Parses markdown text out of an XML node 
 		 * 
 		 * @param child			xml node we want to parse into Markdown 
 		 * @param Indent		indenting prefix for each newline of text in the converted markdown
 		 * @param ResolveLink	delegate to call to resolve doxygen id's into link paths. may be null.
 		 * @return				a string containing the text from the child node 
 		 */
-		public static string ParseXml(XmlNode child, string Indent, ResolveLinkDelegate ResolveLink)
+		public static string ParseXml(XmlNode Node, string Indent, ResolveLinkDelegate ResolveLink)
 		{
-			string output = "";
+			MarkdownWriter Writer = new MarkdownWriter(Indent);
+			ConvertNodeContents(Node, Writer, ResolveLink);
+			return Writer.ToString();
+		}
 
-			XmlNode parent = child.ParentNode;
-			if (child.HasChildNodes)
+		/** 
+		 * Parses markdown text out of an XML node 
+		 * 
+		 * @param child			xml node we want to parse into Markdown 
+		 * @param Indent		indenting prefix for each newline of text in the converted markdown
+		 * @param Output		storage for the output string
+		 * @param ResolveLink	delegate to call to resolve doxygen id's into link paths. may be null.
+		 */
+		static void ConvertNode(XmlNode Node, MarkdownWriter Writer, ResolveLinkDelegate ResolveLink)
+		{
+			switch (Node.Name)
 			{
-				switch (child.Name)
-				{
-					case "para":
-						if (parent.Name != "listitem" && parent.Name != "parameterdescription" && parent.Name != "simplesect")
-						{
-							output += Environment.NewLine;
-							output += ParseChildren(child, Indent, ResolveLink);
-						}
-						else if (parent.Name == "listitem")
+				case "ulink":
+					string LinkText = Node.InnerText;
+					string LinkUrl = Node.Attributes.GetNamedItem("url").InnerText;
+
+					const string WebPrefix = "docs.unrealengine.com/latest/INT/";
+					int WebPrefixIdx = LinkUrl.IndexOf(WebPrefix);
+					if (WebPrefixIdx != -1 && LinkText.Contains("docs.unrealengine.com"))
+					{
+						LinkText = "";
+						LinkUrl = LinkUrl.Substring(WebPrefixIdx + WebPrefix.Length);
+						if (LinkUrl.EndsWith(".html"))
 						{
-							output += Environment.NewLine;
-							if (child == parent.FirstChild)
-							{
-								output += Indent + "* " + ParseChildren(child, Indent, ResolveLink);
-							}
-							else
-							{
-								output += Environment.NewLine + Indent + "\t" + ParseChildren(child, Indent, ResolveLink);
-							}
+							LinkUrl = LinkUrl.Substring(0, LinkUrl.LastIndexOf('/') + 1);
 						}
-						else if (parent.Name == "parameterdescription" || parent.Name == "simplesect")
+					}
+
+					Writer.WriteLink(LinkText, LinkUrl);
+					break;
+				case "ref":
+					XmlAttribute RefAttribute = Node.Attributes["refid"];
+					if (RefAttribute != null && ResolveLink != null)
+					{
+						string LinkPath = ResolveLink(RefAttribute.Value);
+						if (LinkPath != null)
 						{
-							output += ParseChildren(child, Indent, ResolveLink);
+							Writer.WriteLink(Node.InnerText, LinkPath);
+							break;
 						}
-						break;
-					case "ulink":
-						if (bDocToolCanParseExternalLinks)
+					}
+					Writer.Write(Node.InnerText);
+					break;
+				case "bold":
+					Writer.Write("**");
+					ConvertNodeContents(Node, Writer, ResolveLink);
+					Writer.Write("**");
+					break;
+				case "emphasis":
+					Writer.Write("_");
+					ConvertNodeContents(Node, Writer, ResolveLink);
+					Writer.Write("_");
+					break;
+				case "computeroutput":
+					Writer.Write("`");
+					ConvertNodeContents(Node, Writer, ResolveLink);
+					Writer.Write("_");
+					break;
+				case "parameterlist":
+					break;
+				case "parameteritem":
+					break;
+				case "parameternamelist":
+					break;
+				case "parameterdescription":
+					ConvertNodeContents(Node, Writer, ResolveLink);
+					break;
+				case "xrefsect":
+					break;
+				case "simplesect":
+					XmlAttribute KindAttribute = Node.Attributes["kind"];
+					if (KindAttribute == null || KindAttribute.Value != "see")
+					{
+						ConvertNodeContents(Node, Writer, ResolveLink);
+					}
+					break;
+				case "para":
+					Writer.WriteParagraph();
+					ConvertNodeContents(Node, Writer, ResolveLink);
+					break;
+				case "itemizedlist":
+					Writer.WriteLine();
+					foreach (XmlNode ListItemNode in Node.SelectNodes("listitem"))
+					{
+						XmlNodeList ParaNodes = ListItemNode.SelectNodes("para");
+						if (ParaNodes.Count > 0)
 						{
-							string LinkText = ParseChildren(child, Indent, ResolveLink);
-							string LinkUrl = child.Attributes.GetNamedItem("url").InnerText;
+							// Write the first node
+							Writer.Write("* ");
+							Writer.AddIndent(" ");
+							ConvertNodeContents(ParaNodes[0], Writer, ResolveLink);
 
-							const string WebPrefix = "docs.unrealengine.com/latest/INT/";
-							int WebPrefixIdx = LinkUrl.IndexOf(WebPrefix);
-							if(WebPrefixIdx != -1 && LinkText.Contains("docs.unrealengine.com"))
+							// Write anything else as an actual paragraph
+							for (int Idx = 1; Idx < ParaNodes.Count; Idx++)
 							{
-								LinkText = "";
-								LinkUrl = LinkUrl.Substring(WebPrefixIdx + WebPrefix.Length);
-								if(LinkUrl.EndsWith(".html"))
-								{
-									LinkUrl = LinkUrl.Substring(0, LinkUrl.LastIndexOf('/') + 1);
-								}
+								Writer.WriteLine();
+								ConvertNodeContents(ParaNodes[Idx], Writer, ResolveLink);
 							}
 
-							output += Indent + String.Format("[{0}]({1})", LinkText, LinkUrl);
-						}
-						else
-						{
-							output += Indent + EscapeText(child.Attributes.GetNamedItem("url").InnerText);
-						}
-						break;
-					case "ref":
-						XmlAttribute RefAttribute = child.Attributes["refid"];
-						if (RefAttribute != null && ResolveLink != null)
-						{
-							string LinkPath = ResolveLink(RefAttribute.Value);
-							if(LinkPath != null)
-							{
-								output += String.Format("[{0}]({1})", ParseChildren(child, Indent, ResolveLink), LinkPath);
-								break;
-							}
+							// Finish the line
+							Writer.RemoveIndent();
+							Writer.WriteLine();
 						}
-						output += ParseChildren(child, Indent, ResolveLink);
-						break;
-					case "bold":
-						output += Indent + String.Format("**{0}**", ParseChildren(child, Indent, ResolveLink));
-						break;
-					case "emphasis":
-						output += Indent + String.Format("_{0}_", ParseChildren(child, Indent, ResolveLink));
-						break;
-					case "computeroutput":
-						output += Indent + String.Format("`{0}`", ParseChildren(child, Indent, ResolveLink));
-						break;
-					case "itemizedlist":
-						if (parent.ParentNode.Name == "listitem")
-						{
-							output += ParseChildren(child, Indent + "\t", ResolveLink);
-						}
-						else
-						{
-							output += Environment.NewLine + ParseChildren(child, Indent, ResolveLink) + Environment.NewLine + Environment.NewLine;
-						}
-						break;
-					case "parameterlist":
-						break;
-					case "parameteritem":
-						break;
-					case "parameternamelist":
-						break;
-					case "simplesect":
-						XmlAttribute KindAttribute = child.Attributes["kind"];
-						if (KindAttribute != null && KindAttribute.Value == "see")
-						{
-							output += "  \n  \nFor more information, see " + ParseChildren(child, Indent + "\t", ResolveLink) + ".";
-						}
-						break;
-					case "xrefsect":
-						break;
-					default:
-						output += ParseChildren(child, Indent, ResolveLink);
-						break;
-				}
-			}
-			else
-			{
-				output += Indent + EscapeText(child.InnerText);
+					}
+					break;
+				default:
+					ConvertNodeContents(Node, Writer, ResolveLink);
+					break;
 			}
-
-			return output;
 		}
 
-		/** Class ParseXML for all child nodes of the given child
+		/** 
+		 * Converts all children of the given XML node
 		 * 
-		 * @param child xml node we want to parse all children into Markdown 
-		 * @return a string containing the text from the child node 
+		 * @param Node			Xml node we want to parse all children into Markdown 
+		 * @param Writer		Output for markdown text
+		 * @param ResolveLink	Delegate to call to resolve doxygen id's into link paths. may be null.
 		 */
-		static String ParseChildren(XmlNode child, string Indent, ResolveLinkDelegate ResolveLink)
+		static void ConvertNodeContents(XmlNode Node, MarkdownWriter Writer, ResolveLinkDelegate ResolveLink)
 		{
-			String output = "";
-
-			foreach (XmlNode childNode in child.ChildNodes)
+			if (Node.HasChildNodes)
+			{
+				foreach (XmlNode ChildNode in Node.ChildNodes)
+				{
+					ConvertNode(ChildNode, Writer, ResolveLink);
+				}
+			}
+			else
 			{
-				if (output.Length > 0 && output[output.Length - 1] != ' ') output += " ";
-				output += ParseXml(childNode, Indent, ResolveLink);
+				Writer.Write(Node.InnerText);
 			}
-
-			return output;
 		}
 
 		/** 
@@ -362,5 +453,38 @@ namespace DoxygenLib
             }
             return -1;
         }
+
+		public static int SkipCharacter(string Text, int Idx)
+		{
+			if(Text[Idx] == '&')
+			{
+				// Try to read a sequence of characters naming the entity
+				int EndIdx = Idx + 1;
+				while(EndIdx < Text.Length && Char.IsLetter(Text[EndIdx]))
+				{
+					EndIdx++;
+				}
+
+				// Check it finishes with a semicolon
+				if(EndIdx < Text.Length && Text[EndIdx] == ';')
+				{
+					return EndIdx + 1;
+				}
+			}
+			else if(Text[Idx] == '[')
+			{
+				// Find the end of the link
+				int EndIdx = Text.IndexOf(']', Idx + 1);
+				if(EndIdx != -1 && EndIdx + 1 < Text.Length && Text[EndIdx + 1] == '(')
+				{
+					EndIdx = Text.IndexOf(')', EndIdx + 1);
+					if(EndIdx != -1)
+					{
+						return EndIdx + 1;
+					}
+				}
+			}
+			return Idx + 1;
+		}
 	}
 }
