Add documentation build into GUBP

diff --git a/Engine/Source/Programs/AutomationTool/Scripts/BuildDocumentation.Automation.cs b/Engine/Source/Programs/AutomationTool/Scripts/BuildDocumentation.Automation.cs
--- a/Engine/Source/Programs/AutomationTool/Scripts/BuildDocumentation.Automation.cs
+++ b/Engine/Source/Programs/AutomationTool/Scripts/BuildDocumentation.Automation.cs
@@ -8,348 +8,136 @@ using System.Reflection;
 using AutomationTool;
 using UnrealBuildTool;
 
-[Help(@"Builds engine documentation.")]
-[Help("-env", "Builds the compile environment")]
-public class BuildDocumentation : BuildCommand
+class ToolsForDocumentationNode : GUBP.CompileNode
 {
-	public override void ExecuteBuild()
+    public ToolsForDocumentationNode(UnrealTargetPlatform InHostPlatform)
+        : base(InHostPlatform, false)
+    {
+		AgentSharingGroup = "Documentation" + StaticGetHostPlatformSuffix(InHostPlatform);
+    }
+
+    public static string StaticGetFullName(UnrealTargetPlatform InHostPlatform)
+    {
+		return "ToolsForDocumentation" + StaticGetHostPlatformSuffix(InHostPlatform);
+    }
+
+	public override string GetFullName()
+    {
+        return StaticGetFullName(HostPlatform);
+    }
+
+	public override UE4Build.BuildAgenda GetAgenda(GUBP bp)
+    {
+        var Agenda = new UE4Build.BuildAgenda();
+		Agenda.DotNetSolutions.Add("Engine/Source/Programs/UnrealDocTool/APIDocTool/APIDocTool.sln");
+		Agenda.DotNetSolutions.Add("Engine/Source/Programs/UnrealDocTool/UnrealDocTool/UnrealDocTool.sln");
+        return Agenda;
+    }
+}
+	
+class DocumentationNode : GUBP.GUBPNode
+{
+	public DocumentationNode(UnrealTargetPlatform InHostPlatform)
 	{
-		// Parse the command line
-		bool bClean = ParseParam("clean");
-		bool bCleanMeta = bClean || ParseParam("cleanmeta");
-		bool bCleanEnv = bClean || ParseParam("cleanenv");
-		bool bCleanXml = bClean || ParseParam("cleanxml");
-		bool bCleanJson = bClean || ParseParam("cleanjson");
-		bool bCleanUdn = bClean || ParseParam("cleanudn");
-		bool bCleanHtml = bClean || ParseParam("cleanhtml");
-		bool bCleanChm = bClean || ParseParam("cleanchm");
-
-		bool bBuild = ParseParam("build");
-		bool bBuildMeta = bBuild || ParseParam("buildmeta");
-		bool bBuildEnv = bBuild || ParseParam("buildenv");
-		bool bBuildXml = bBuild || ParseParam("buildxml");
-		bool bBuildJson = bBuild || ParseParam("buildjson");
-		bool bBuildUdn = bBuild || ParseParam("buildudn");
-		bool bBuildHtml = bBuild || ParseParam("buildhtml");
-		bool bBuildChm = bBuild || ParseParam("buildchm");
-
-		if(ParseParam("nodoxygen"))
-		{
-			bCleanMeta = bCleanEnv = bCleanXml = bBuildMeta = bBuildEnv = bBuildXml = false;
-		}
-
-		bool bStats = ParseParam("stats");
-
-		bool bMakeArchives = ParseParam("archive");
-
-		string Filter = ParseParamValue("filter");
-
-		// Create the intermediate folders
-		string IntermediateDir = Path.Combine(CmdEnv.LocalRoot, "Engine\\Intermediate\\Documentation");
-		CommandUtils.CreateDirectory(IntermediateDir);
-		string BuildDir = Path.Combine(IntermediateDir, "build");
-		CommandUtils.CreateDirectory(BuildDir);
-		string DoxygenDir = Path.Combine(IntermediateDir, "doxygen");
-		CommandUtils.CreateDirectory(DoxygenDir);
-		string JsonDir = Path.Combine(IntermediateDir, "json");
-		CommandUtils.CreateDirectory(JsonDir);
-		string MetadataDir = Path.Combine(IntermediateDir, "metadata");
-		CommandUtils.CreateDirectory(MetadataDir);
-		string ArchiveDir = Path.Combine(CmdEnv.LocalRoot, "Engine\\Documentation\\Builds");
-		CommandUtils.CreateDirectory(ArchiveDir);
-
-		// Get all the file paths
-		string TargetInfoPath = Path.Combine(BuildDir, "targetinfo.xml");
-		string UnrealBuildToolProject = Path.Combine(CmdEnv.LocalRoot, "Engine/Source/Programs/UnrealBuildTool/UnrealBuildTool.csproj");
-		string ApiDocToolSolution = Path.Combine(CmdEnv.LocalRoot, "Engine/Source/Programs/UnrealDocTool/APIDocTool/APIDocTool.sln");
-		string ApiDocToolPath = Path.Combine(CmdEnv.LocalRoot, "Engine/Source/Programs/UnrealDocTool/APIDocTool/APIDocTool/bin/x64/Release/APIDocTool.exe");
-		string UnrealDocToolSolution = Path.Combine(CmdEnv.LocalRoot, "Engine/Source/Programs/UnrealDocTool/UnrealDocTool/UnrealDocTool.sln");
-		string StatsPath = "\\\\epicgames.net\\root\\UE3\\UE4\\Documentation\\API-Stats.txt";
-
-		// Compile UnrealBuildTool
-		MsBuild(CmdEnv, UnrealBuildToolProject, "/verbosity:normal /target:Rebuild /property:Configuration=Development /property:Platform=AnyCPU", "BuildUnrealBuildTool");
-
-		// Compile APIDocTool
-		RunAndLog(CmdEnv, CmdEnv.MsDevExe, String.Format("\"{0}\" /rebuild \"Release|x64\"", ApiDocToolSolution), "BuildApiDocTool");
-
-		// Compile UnrealDocTool
-		RunAndLog(CmdEnv, CmdEnv.MsDevExe, String.Format("\"{0}\" /rebuild \"Release|Any CPU\"", UnrealDocToolSolution), "BuildUnrealDocTool");
-
-		// Prepare the shared doc tool command line, containing all the paths
-		string ApiToolCommandLine = "";
-		ApiToolCommandLine += " -enginedir=\"" + Path.Combine(CmdEnv.LocalRoot, "Engine") + "\"";
-		ApiToolCommandLine += " -targetinfo=\"" + TargetInfoPath + "\"";
-		ApiToolCommandLine += " -xmldir=\"" + DoxygenDir + "\"";
-		ApiToolCommandLine += " -jsondir=\"" + JsonDir + "\"";
-		ApiToolCommandLine += " -metadatadir=\"" + MetadataDir + "\"";
-		if (Filter != null) ApiToolCommandLine += " -filter=" + Filter;
-
-		// Execute the clean
-		if(bCleanEnv)
-		{
-			RunAndLog(CmdEnv, ApiDocToolPath, "-cleantarget" + ApiToolCommandLine, "ApiDocTool-CleanTargetInfo");
-		}
-		if (bCleanMeta)
-		{
-			RunAndLog(CmdEnv, ApiDocToolPath, "-cleanmeta" + ApiToolCommandLine, "ApiDocTool-CleanMeta");
-		}
-		if(bCleanXml)
-		{
-			RunAndLog(CmdEnv, ApiDocToolPath, "-cleanxml" + ApiToolCommandLine, "ApiDocTool-CleanXML");
-		}
-		if (bCleanJson)
-		{
-			RunAndLog(CmdEnv, ApiDocToolPath, "-cleanjson" + ApiToolCommandLine, "ApiDocTool-CleanJSON");
-		}
-		if (bCleanUdn)
-		{
-			RunAndLog(CmdEnv, ApiDocToolPath, "-cleanudn" + ApiToolCommandLine, "ApiDocTool-CleanUDN");
-		}
-		if(bCleanHtml)
-		{
-			RunAndLog(CmdEnv, ApiDocToolPath, "-cleanhtml" + ApiToolCommandLine, "ApiDocTool-CleanHTML");
-		}
-		if (bCleanChm)
-		{
-			RunAndLog(CmdEnv, ApiDocToolPath, "-cleanchm" + ApiToolCommandLine, "ApiDocTool-CleanCHM");
-		}
-
-		// Execute the build
-		if (bBuildEnv)
-		{
-			RunAndLog(CmdEnv, ApiDocToolPath, "-buildtarget" + ApiToolCommandLine, "ApiDocTool-BuildTargetInfo");
-		}
-		if (bBuildMeta)
-		{
-			RunAndLog(CmdEnv, ApiDocToolPath, "-buildmeta" + ApiToolCommandLine, "ApiDocTool-BuildMeta");
-		}
-		if (bBuildXml)
-		{
-			RunAndLog(CmdEnv, ApiDocToolPath, "-buildxml" + ApiToolCommandLine, "ApiDocTool-BuildXML");
-			if (bMakeArchives)
-			{
-				List<string> InputFiles = new List<string>();
-				FindFilesForTar(IntermediateDir, "build", InputFiles, true);
-				FindFilesForTar(IntermediateDir, "doxygen", InputFiles, true);
-				FindFilesForTar(IntermediateDir, "metadata", InputFiles, true);
-
-				CreateAndSubmitArchiveFromFiles("Engine/Documentation/Builds/API-XML.tgz", "API documentation intermediates", IntermediateDir, InputFiles);
-			}
-		}
-		if (bBuildJson)
-		{
-			RunAndLog(CmdEnv, ApiDocToolPath, "-buildjson" + ApiToolCommandLine, "ApiDocTool-BuildJSON");
-			if (bMakeArchives)
-			{
-				CreateAndSubmitArchiveFromDir("Engine/Documentation/Builds/API-JSON.tgz", "Blueprint API documentation intermediates", Path.Combine(IntermediateDir, "json"));
-			}
-		}
-		if (bBuildUdn)
-		{
-			RunAndLog(CmdEnv, ApiDocToolPath, "-buildudn" + (bStats? (" -stats=\"" + StatsPath + "\"") : "") + ApiToolCommandLine, "ApiDocTool-BuildUDN");
-			if (bMakeArchives)
-			{
-				string BaseDir = Path.Combine(CmdEnv.LocalRoot, "Engine\\Documentation\\Source");
-
-				List<string> InputFiles = new List<string>();
-				FindFilesForTar(BaseDir, "API", InputFiles, true);
-				FindFilesForTar(BaseDir, "BlueprintAPI", InputFiles, true);
-
-				CreateAndSubmitArchiveFromFiles("Engine/Documentation/Builds/API-UDN.tgz", "API documentation UDN output", BaseDir, InputFiles);
-				CreateAndSubmitArchiveFromFiles("Engine/Documentation/Builds/API-Sitemap.tgz", "API documentation sitemap output", Path.Combine(CmdEnv.LocalRoot, "Engine\\Documentation\\CHM"), new string[] { "API.hhc", "API.hhk", "BlueprintAPI.hhc", "BlueprintAPI.hhk" });
-			}
-		}
-		if (bBuildHtml)
-		{
-			RunAndLog(CmdEnv, ApiDocToolPath, "-buildhtml" + ApiToolCommandLine, "ApiDocTool-BuildHTML");
-			if (bMakeArchives)
-			{
-				string BaseDir = Path.Combine(CmdEnv.LocalRoot, "Engine\\Documentation\\HTML");
-
-				List<string> InputFiles = new List<string>();
-				FindFilesForTar(BaseDir, "Include", InputFiles, true);
-				FindFilesForTar(BaseDir, "INT\\API", InputFiles, true);
-				FindFilesForTar(BaseDir, "Images", InputFiles, false);
-
-				CommandUtils.CopyFile(Path.Combine(CmdEnv.LocalRoot, "Engine\\Documentation\\Extras\\API\\index.html"), Path.Combine(BaseDir, "index.html"));
-				InputFiles.Add("index.html");
-
-				CreateAndSubmitArchiveFromFiles("Engine/Documentation/Builds/API-HTML.tgz", "API documentation HTML output", BaseDir, InputFiles);
-			}
-		}
-		if (bBuildChm)
-		{
-			RunAndLog(CmdEnv, ApiDocToolPath, "-buildchm" + ApiToolCommandLine, "ApiDocTool-BuildCHM");
-			SubmitFile("Engine/Documentation/CHM/API.chm", "API documentation CHM output");
-			SubmitFile("Engine/Documentation/CHM/BlueprintAPI.chm", "Blueprint API documentation CHM output");
-		}
+		AgentSharingGroup = "Documentation" + GUBP.HostPlatformNode.StaticGetHostPlatformSuffix(InHostPlatform);
+		AddDependency(ToolsForDocumentationNode.StaticGetFullName(InHostPlatform));
 	}
 
-	static void CreateAndSubmitArchiveFromDir(string DepotArchivePath, string Description, string BaseDir)
+	public override int CISFrequencyQuantumShift(GUBP bp)
 	{
-		// Find all the input files
-		List<string> InputFiles = new List<string>();
-		FindFilesForTar(new DirectoryInfo(BaseDir), "", InputFiles, true);
-		CreateAndSubmitArchiveFromFiles(DepotArchivePath, Description, BaseDir, InputFiles.ToArray());
+		return base.CISFrequencyQuantumShift(bp);
 	}
 
-	static void CreateAndSubmitArchiveFromFiles(string DepotArchivePath, string Description, string BaseDir, IEnumerable<string> InputFiles)
+	protected void ExecuteApiDocTool(string Arguments, string LogName)
 	{
-		string ArchivePath = Path.Combine(CmdEnv.LocalRoot, DepotArchivePath);
-
-		// Create the archive
-		CreateTgzFromFiles(ArchivePath, BaseDir, InputFiles);
-
-		// Submit the archive to P4
-		SubmitFile(DepotArchivePath, Description);
+		string ApiDocToolPath = Path.Combine(CommandUtils.CmdEnv.LocalRoot, "Engine/Source/Programs/UnrealDocTool/APIDocTool/APIDocTool/bin/x64/Release/APIDocTool.exe");
+		string ApiToolCommandLine = Arguments + " -enginedir=\"" + Path.Combine(CommandUtils.CmdEnv.LocalRoot, "Engine") + "\"";
+		CommandUtils.RunAndLog(CommandUtils.CmdEnv, ApiDocToolPath, ApiToolCommandLine, LogName);
 	}
 
-	static void SubmitFile(string DepotPath, string Description)
+	public static void SubmitOutputs(string Description, params string[] FileSpecs)
 	{
-		if (P4Enabled)
+		if (CommandUtils.P4Enabled)
 		{
-			int Changelist = P4.CreateChange(P4Env.Client, String.Format("{0} from CL#{1}", Description, P4Env.Changelist));
-			P4.Reconcile(Changelist, CombinePaths(PathSeparator.Slash, P4Env.ClientRoot, DepotPath));
+			int Changelist = CommandUtils.P4.CreateChange(CommandUtils.P4Env.Client, String.Format("{0} from CL#{1}", Description, CommandUtils.P4Env.Changelist));
+
+			foreach (string FileSpec in FileSpecs)
+			{
+				CommandUtils.P4.Reconcile(Changelist, CommandUtils.CombinePaths(PathSeparator.Slash, CommandUtils.P4Env.ClientRoot, FileSpec));
+			}
 
-			if (!P4.TryDeleteEmptyChange(Changelist))
+			if (!CommandUtils.P4.TryDeleteEmptyChange(Changelist))
 			{
 				if (!GlobalCommandLine.NoSubmit)
 				{
 					int SubmittedChangelist;
-					P4.Submit(Changelist, out SubmittedChangelist, true, true);
+					CommandUtils.P4.Submit(Changelist, out SubmittedChangelist, true, true);
 				}
 			}
 		}
 	}
+}
 
-	static void CreateTgzFromFiles(string TarPath, string BaseDir, IEnumerable<string> InputFiles)
+class CodeDocumentationNode : DocumentationNode
+{
+	public CodeDocumentationNode(UnrealTargetPlatform InHostPlatform)
+		: base(InHostPlatform)
 	{
-		Console.WriteLine("Creating archive '{0}'...", TarPath);
-		if (File.Exists(TarPath))
-		{
-			CommandUtils.DeleteFile(TarPath);
-		}
-		using(FileStream TarFileStream = new FileStream(TarPath, FileMode.Create))
-		{
-			using (GZipStream ZipStream = new GZipStream(TarFileStream, CompressionMode.Compress))
-			{
-				using (BinaryWriter TarWriter = new BinaryWriter(ZipStream))
-				{
-					foreach (string InputFile in InputFiles)
-					{
-						byte[] FileData = File.ReadAllBytes(Path.Combine(BaseDir, InputFile));
-						WriteFileToTar(TarWriter, InputFile, FileData);
-					}
-					TarWriter.Write(new byte[1024]);
-				}
-			}
-		}
 	}
 
-	static void FindFilesForTar(string BaseDir, string BaseTarPath, List<string> Files, bool bRecursive)
+	public override void DoBuild(GUBP bp)
 	{
-		FindFilesForTar(new DirectoryInfo(Path.Combine(BaseDir, BaseTarPath)), BaseTarPath, Files, bRecursive);
+		ExecuteApiDocTool("-rebuildcode -filter=Messaging/...", "APIDocTool-Code");
+		SubmitOutputs("Code documentation", "Engine/Documentation/Builds/CodeAPI-*", "Engine/Documentation/CHM/API.chm");
+		base.DoBuild(bp);
 	}
 
-	static void FindFilesForTar(DirectoryInfo Dir, string BaseTarPath, List<string> Files, bool bRecursive)
+	public override string GetFullName()
 	{
-		if (bRecursive)
-		{
-			foreach (DirectoryInfo ChildDir in Dir.GetDirectories())
-			{
-				FindFilesForTar(ChildDir, Path.Combine(BaseTarPath, ChildDir.Name), Files, bRecursive);
-			}
-		}
-		foreach (FileInfo FileInfo in Dir.GetFiles())
-		{
-			Files.Add(Path.Combine(BaseTarPath, FileInfo.Name));
-		}
+		return StaticGetFullName();
 	}
 
-	const int TarFileNameFieldLength = 99;
-
-	static void WriteFileToTar(BinaryWriter Writer, string FilePath, byte[] FileData)
+	public static string StaticGetFullName()
 	{
-		string NormalizedFilePath = FilePath.Replace('\\', '/');
-
-		// Generate a long header record if necessary
-		if (NormalizedFilePath.Length > TarFileNameFieldLength)
-		{
-			WriteTarHeader(Writer, "././@LongLink", NormalizedFilePath.Length + 1, 'L');
-
-			byte[] FilePathBytes = new byte[NormalizedFilePath.Length + 1];
-			EncodeTarString(FilePathBytes, 0, FilePathBytes.Length, NormalizedFilePath);
-
-			WriteTarBlocks(Writer, FilePathBytes);
-		}
-
-		// Now write the normal file
-		WriteTarHeader(Writer, NormalizedFilePath, FileData.Length, '0');
-		WriteTarBlocks(Writer, FileData);
+		return "CodeDocumentation";
 	}
+}
 
-	static void WriteTarHeader(BinaryWriter Writer, string FileName, int Length, char Type)
+class BlueprintDocumentationNode : DocumentationNode
+{
+	public BlueprintDocumentationNode(UnrealTargetPlatform InHostPlatform)
+		: base(InHostPlatform)
 	{
-		// Write the name
-		byte[] Header = new byte[512];
-
-		// Name
-		EncodeTarString(Header, 0, TarFileNameFieldLength, FileName);
-
-		// File mode
-		EncodeTarValue(Header, 100, 7, 0x81ff);
-
-		// UID
-		EncodeTarValue(Header, 108, 7, 0);
-
-		// GID
-		EncodeTarValue(Header, 116, 7, 0);
-
-		// File size
-		EncodeTarValue(Header, 124, 12, Length);
-
-		// File timestamp
-		EncodeTarValue(Header, 136, 12, 0);
-
-		// File type ('0' = Regular file)
-		Header[156] = (byte)Type;
-
-		// Generate the the checksum (seed the checksum as if the field was initialized with 8 spaces)
-		int Checksum = (byte)' ' * 8;
-		for (int Idx = 0; Idx < Header.Length; Idx++) Checksum += Header[Idx];
-		EncodeTarValue(Header, 148, 7, Checksum);
+		AddDependency(GUBP.RootEditorNode.StaticGetFullName(InHostPlatform));
+	}
 
-		// Write the header
-		Writer.Write(Header);
+	public override void DoBuild(GUBP bp)
+	{
+		ExecuteApiDocTool("-rebuildblueprint", "APIDocTool-Blueprint");
+		SubmitOutputs("Code documentation", "Engine/Documentation/Builds/BlueprintAPI-*", "Engine/Documentation/CHM/BlueprintAPI.chm");
+		base.DoBuild(bp);
 	}
 
-	static void WriteTarBlocks(BinaryWriter Writer, byte[] Data)
+	public override string GetFullName()
 	{
-		Writer.Write(Data);
-		Writer.Write(new byte[(512 - (Data.Length % 512)) % 512]);
+		return StaticGetFullName();
 	}
 
-	static void EncodeTarString(byte[] Data, int Offset, int Length, string Text)
+	public static string StaticGetFullName()
 	{
-		// Convert the Text to ASCII
-		for (int Idx = 0; Idx < Math.Min(Length, Text.Length); Idx++)
-		{
-			if (Text[Idx] == 0 || Text[Idx] > 0x7f)
-			{
-				throw new AutomationException("Invalid character in TAR string");
-			}
-			else
-			{
-				Data[Offset + Idx] = (byte)Text[Idx];
-			}
-		}
+		return "BlueprintDocumentation";
 	}
+}
 
-	static void EncodeTarValue(byte[] Data, int Offset, int Length, int Value)
+public class DocumentationNodeAdder : GUBP.GUBPNodeAdder
+{
+	public override void AddNodes(GUBP bp, UnrealTargetPlatform InHostPlatform)
 	{
-		// Write the value as an octal string followed by a space
-		string ValueString = Convert.ToString(Value, 8) + " ";
-		while (ValueString.Length < Length) ValueString = " " + ValueString;
-		EncodeTarString(Data, Offset, Length, ValueString);
+		if(InHostPlatform == UnrealTargetPlatform.Win64)
+		{
+			bp.AddNode(new ToolsForDocumentationNode(InHostPlatform));
+			bp.AddNode(new CodeDocumentationNode(InHostPlatform));
+			bp.AddNode(new BlueprintDocumentationNode(InHostPlatform));
+		}
 	}
 }
diff --git a/Engine/Source/Programs/UnrealDocTool/APIDocTool/APIDocTool/Program.cs b/Engine/Source/Programs/UnrealDocTool/APIDocTool/APIDocTool/Program.cs
--- a/Engine/Source/Programs/UnrealDocTool/APIDocTool/APIDocTool/Program.cs
+++ b/Engine/Source/Programs/UnrealDocTool/APIDocTool/APIDocTool/Program.cs
@@ -20,6 +20,14 @@ namespace APIDocTool
 {
     public class Program
 	{
+		[Flags]
+		enum BuildActions
+		{
+			Clean = 0x01,
+			Build = 0x02,
+			Archive = 0x04,
+		}
+
 		[DllImport("kernel32.dll")]
 		public static extern uint SetErrorMode(uint hHandle);
 
@@ -188,352 +196,221 @@ namespace APIDocTool
 			DoxygenPredefinedMacros = DoxygenPredefinedMacros.Union(DelegateMacros).ToArray();
 		}
 
-		static void Main(string[] Arguments)
+		static int Main(string[] Arguments)
         {
-			string EngineRootDir = Path.GetFullPath(Path.Combine(Path.GetDirectoryName(Assembly.GetExecutingAssembly().Location), "..\\..\\..\\..\\..\\..\\..\\.."));
-			string IntermediateRootDir = Path.Combine(EngineRootDir, "Intermediate\\Documentation");
-
-			bool bValidArgs = true;
-
-			bool bCleanTargetInfo = false;
-			bool bCleanMetadata = false;
-			bool bCleanXml = false;
-			bool bCleanJson = false;
-			bool bCleanUdn = false;
-			bool bCleanHtml = false;
-			bool bCleanChm = false;
-
-			bool bBuildTargetInfo = false;
-			bool bBuildMetadata = false;
-			bool bBuildXml = false;
-			bool bBuildJson = false;
-			bool bBuildUdn = false;
-			bool bBuildHtml = false;
-			bool bBuildChm = false;
-
-			bool bBuildCodeAPI = false;
-			bool bBuildBlueprintAPI = false;
-
-			string TargetInfoPath = Path.Combine(IntermediateRootDir, "build\\targetinfo.xml");
-			string EngineDir = EngineRootDir;
-			string MetadataDir = Path.Combine(IntermediateRootDir, "metadata");
-			string XmlDir = Path.Combine(IntermediateRootDir, "doxygen");
-			string JsonDir = Path.Combine(IntermediateRootDir, "json");
-			string StatsPath = null;
+			List<string> ArgumentList = new List<string>(Arguments);
 
-			List<string> Filters = new List<string>();
+			// Get the default paths
+			string EngineDir = ParseArgumentDirectory(ArgumentList, "-enginedir=", Path.Combine(Path.GetDirectoryName(Assembly.GetExecutingAssembly().Location), "..\\..\\..\\..\\..\\..\\..\\.."));
+			string IntermediateDir = ParseArgumentDirectory(ArgumentList, "-intermediatedir=", Path.Combine(EngineDir, "Intermediate\\Documentation"));
+			string DocumentationDir = ParseArgumentDirectory(ArgumentList, "-documentationdir=", Path.Combine(EngineDir, "Documentation"));
 
-			// Parse the command line
-			foreach (string Argument in Arguments)
-			{
-				int EqualsIdx = Argument.IndexOf('=');
-				string OptionName = (EqualsIdx == -1)? Argument : Argument.Substring(0, EqualsIdx + 1);
-				string OptionValue = (EqualsIdx == -1)? null : Argument.Substring(EqualsIdx + 1);
+			// Check if we're just building an index, no actual content
+			bool bIndexOnly = ArgumentList.Remove("-indexonly");
 
-				if (OptionName == "-blueprint")
-				{
-					bBuildBlueprintAPI = true;
-				}
-				else if (OptionName == "-code")
-				{
-					bBuildCodeAPI = true;
-				}
-				else if (OptionName == "-clean")
-				{
-					bCleanTargetInfo = bCleanMetadata = bCleanXml = bCleanUdn = bCleanHtml = true;
-				}
-				else if (OptionName == "-build")
-				{
-					bBuildTargetInfo = bBuildMetadata = bBuildXml = bBuildUdn = bBuildHtml = true;
-				}
-				else if (OptionName == "-rebuild")
-				{
-					bCleanTargetInfo = bCleanMetadata = bCleanXml = bCleanUdn = bCleanHtml = bCleanChm = true;
-					bBuildTargetInfo = bBuildMetadata = bBuildXml = bBuildUdn = bBuildHtml = bBuildChm = true;
-				}
-				else if (OptionName == "-cleantarget")
-				{
-					bCleanTargetInfo = true;
-				}
-				else if (OptionName == "-buildtarget")
-				{
-					bBuildTargetInfo = true;
-				}
-				else if (OptionName == "-cleanmeta")
-				{
-					bCleanMetadata = true;
-				}
-				else if (OptionName == "-buildmeta")
-				{
-					bBuildMetadata = true;
-				}
-				else if (OptionName == "-rebuildmeta")
-				{
-					bCleanMetadata = bBuildMetadata = true;
-				}
-				else if (OptionName == "-cleanxml")
-				{
-					bCleanXml = true;
-				}
-				else if (OptionName == "-buildxml")
-				{
-					bBuildXml = true;
-				}
-				else if (OptionName == "-rebuildxml")
-				{
-					bCleanXml = bBuildXml = true;
-				}
-				else if (OptionName == "-cleanjson")
-				{
-					bCleanJson = true;
-				}
-				else if (OptionName == "-buildjson")
-				{
-					bBuildJson = true;
-				}
-				else if (OptionName == "-rebuildjson")
-				{
-					bCleanJson = bBuildJson = true;
-				}
-				else if (OptionName == "-cleanudn")
-				{
-					bCleanUdn = true;
-				}
-				else if (OptionName == "-buildudn")
-				{
-					bBuildUdn = true;
-				}
-				else if (OptionName == "-rebuildudn")
-				{
-					bCleanUdn = bBuildUdn = true;
-				}
-				else if (OptionName == "-cleanhtml")
-				{
-					bCleanHtml = true;
-				}
-				else if (OptionName == "-buildhtml")
-				{
-					bBuildHtml = true;
-				}
-				else if (OptionName == "-rebuildhtml")
-				{
-					bCleanHtml = bBuildHtml = true;
-				}
-				else if (OptionName == "-cleanchm")
-				{
-					bCleanChm = true;
-				}
-				else if (OptionName == "-buildchm")
-				{
-					bBuildChm = true;
-				}
-				else if (OptionName == "-rebuildchm")
-				{
-					bCleanChm = bBuildChm = true;
-				}
-				else if (OptionName == "-targetinfo=")
-				{
-					TargetInfoPath = Path.GetFullPath(OptionValue);
-				}
-				else if (OptionName == "-enginedir=")
-				{
-					EngineDir = Path.GetFullPath(OptionValue);
-				}
-				else if (OptionName == "-xmldir=")
-				{
-					XmlDir = Path.GetFullPath(OptionValue);
-				}
-				else if (OptionName == "-jsondir=")
-				{
-					JsonDir = Path.GetFullPath(OptionValue);
-				}
-				else if (OptionName == "-metadatadir=")
-				{
-					MetadataDir = Path.GetFullPath(OptionValue);
-				}
-				else if (OptionName == "-stats=")
-				{
-					StatsPath = Path.GetFullPath(OptionValue);
-				}
-				else if (OptionName == "-indexonly")
-				{
-					bIndexOnly = true;
-				}
-				else if (OptionName == "-filter=")
-				{
-					Filters.AddRange(OptionValue.Split(',').Select(x => x.Replace('\\', '/').Trim()));
-				}
-				else
+			// Remove all the filter arguments
+			List<string> Filters = new List<string>();
+			for (int Idx = 0; Idx < ArgumentList.Count; Idx++)
+			{
+				const string FilterPrefix = "-filter=";
+				if (ArgumentList[Idx].StartsWith(FilterPrefix))
 				{
-					Console.WriteLine("Invalid argument: '{0}'", OptionName);
-					bValidArgs = false;
+					Filters.AddRange(ArgumentList[Idx].Substring(FilterPrefix.Length).Split(',').Select(x => x.Replace('\\', '/').Trim()));
+					ArgumentList.RemoveAt(Idx--);
 				}
 			}
 
-			// If neither -blueprint or -code was specified, build both
-			if (!bBuildCodeAPI && !bBuildBlueprintAPI)
+			// Remove all the arguments for what to build
+			BuildActions DefaultActions;
+			if(ArgumentList.Count == 0)
 			{
-				bBuildCodeAPI = true;
-				bBuildBlueprintAPI = true;
+				DefaultActions = BuildActions.Clean | BuildActions.Build | BuildActions.Archive;
 			}
-
-			// Check we have all the required parameters
-			if (bBuildXml && TargetInfoPath == null)
+			else
 			{
-				Console.WriteLine("Missing -buildenvironment parameter to be able to build XML");
+				DefaultActions = ParseBuildActions(ArgumentList, "");
 			}
-			else if (bValidArgs && EngineDir != null)
+
+			BuildActions CodeActions = ParseBuildActions(ArgumentList, "code") | DefaultActions;
+			BuildActions CodeTargetActions = ParseBuildActions(ArgumentList, "codetarget") | CodeActions;
+			BuildActions CodeMetadataActions = ParseBuildActions(ArgumentList, "codemeta") | CodeActions;
+			BuildActions CodeXmlActions = ParseBuildActions(ArgumentList, "codexml") | CodeActions;
+			BuildActions CodeUdnActions = ParseBuildActions(ArgumentList, "codeudn") | CodeActions;
+			BuildActions CodeHtmlActions = ParseBuildActions(ArgumentList, "codehtml") | CodeActions;
+			BuildActions CodeChmActions = ParseBuildActions(ArgumentList, "codechm") | CodeActions;
+
+			BuildActions BlueprintActions = ParseBuildActions(ArgumentList, "blueprint") | DefaultActions;
+			BuildActions BlueprintJsonActions = ParseBuildActions(ArgumentList, "blueprintjson") | BlueprintActions;
+			BuildActions BlueprintUdnActions = ParseBuildActions(ArgumentList, "blueprintudn") | BlueprintActions;
+			BuildActions BlueprintHtmlActions = ParseBuildActions(ArgumentList, "blueprinthtml") | BlueprintActions;
+			BuildActions BlueprintChmActions = ParseBuildActions(ArgumentList, "blueprintchm") | BlueprintActions;
+
+			// Print the help message and exit
+			if(Arguments.Length == 0 || ArgumentList.Count > 0)
 			{
-				// Derive all the engine paths we need in one place
-				string DoxygenPath = Path.Combine(EngineDir, "Extras\\NotForLicensees\\Doxygen\\bin\\doxygen.exe");
-				string UdnDir = Path.Combine(EngineDir, "Documentation\\Source");
-				string HtmlDir = Path.Combine(EngineDir, "Documentation\\HTML");
-				string ChmDir = Path.Combine(EngineDir, "Documentation\\CHM");
-				string EditorPath = Path.Combine(EngineDir, "Binaries\\Win64\\UE4Editor.exe");
-				string DocToolPath = Path.Combine(EngineDir, "Binaries\\DotNET\\UnrealDocTool.exe");
-				string ChmCompilerPath = Path.Combine(EngineDir, "Extras\\NotForLicensees\\HTML Help Workshop\\hhc.exe");
-				string SettingsPath = Path.Combine(EngineDir, "Documentation\\Extras\\API\\API.ini");
-				string MetadataPath = Path.Combine(MetadataDir, "metadata.xml");
-
-				if (bBuildCodeAPI)
+				foreach(string InvalidArgument in ArgumentList)
 				{
-					// If we don't intermediate paths, make them up
-					if (XmlDir == null) XmlDir = Path.Combine(EngineDir, "Intermediate\\Documentation\\Default\\Xml");
-					if (MetadataDir == null) MetadataDir = Path.Combine(EngineDir, "Intermediate\\Documentation\\Default\\Metadata");
-
-					// Read the settings file
-					Settings = IniFile.Read(SettingsPath);
-					IgnoredFunctionMacros = new HashSet<string>(Settings.FindValueOrDefault("Input.IgnoredFunctionMacros", "").Split('\n'));
-					IncludedSourceFiles = new HashSet<string>(Settings.FindValueOrDefault("Output.IncludedSourceFiles", "").Split('\n'));
-
-					// Find all the metadata pages
-					AddMetadataKeyword(UdnDir, "UCLASS", "Programming/UnrealArchitecture/Reference/Classes#classdeclaration", "Programming/UnrealArchitecture/Reference/Classes/Specifiers");
-					AddMetadataKeyword(UdnDir, "UFUNCTION", "Programming/UnrealArchitecture/Reference/Functions", "Programming/UnrealArchitecture/Reference/Functions/Specifiers");
-					AddMetadataKeyword(UdnDir, "UPROPERTY", "Programming/UnrealArchitecture/Reference/Properties", "Programming/UnrealArchitecture/Reference/Properties/Specifiers");
-					AddMetadataKeyword(UdnDir, "USTRUCT", "Programming/UnrealArchitecture/Reference/Structs", "Programming/UnrealArchitecture/Reference/Structs/Specifiers");
-
-					// Clean the output folders
-					if (bCleanTargetInfo)
-					{
-						CleanTargetInfo(TargetInfoPath);
-					}
-					if (bCleanMetadata)
-					{
-						CleanMetadata(MetadataDir);
-					}
-					if (bCleanXml)
-					{
-						CleanXml(XmlDir);
-					}
-					if (bCleanUdn)
-					{
-						CleanUdn(Path.Combine(UdnDir, "API"));
-					}
-					if (bCleanHtml)
-					{
-						CleanHtml(Path.Combine(HtmlDir, "INT\\API"));
-					}
-					if (bCleanChm)
-					{
-						CleanChm(ChmDir);
-					}
-
-					// Build the data
-					if (!bBuildTargetInfo || BuildTargetInfo(TargetInfoPath, EngineDir))
-					{
-						if (!bBuildMetadata || BuildMetadata(DoxygenPath, EngineDir, MetadataDir, MetadataPath))
-						{
-							if (!bBuildXml || BuildXml(EngineDir, TargetInfoPath, DoxygenPath, XmlDir, Filters))
-							{
-								if (!bBuildUdn || BuildUdn(EngineDir, XmlDir, UdnDir, ChmDir, MetadataPath, StatsPath, Filters))
-								{
-									if (!bBuildHtml || BuildHtml(EngineDir, DocToolPath, UdnDir, HtmlDir, false))
-									{
-										if (!bBuildChm || BuildChm(ChmCompilerPath, HtmlDir, ChmDir, false))
-										{
-											Console.WriteLine("Complete.");
-										}
-									}
-								}
-							}
-						}
-					}
+					Console.WriteLine("Invalid argument: {0}", InvalidArgument);
 				}
-
-				if (bBuildBlueprintAPI)
-				{
-					// If we don't intermediate paths, make them up
-					if (JsonDir == null) JsonDir = Path.Combine(EngineDir, "Intermediate\\Documentation\\Default\\Json");
-
-					// Read the settings file
-					Settings = IniFile.Read(SettingsPath);
-
-					// Clean the output folders
-					if (bCleanTargetInfo)
-					{
-						CleanTargetInfo(TargetInfoPath);
-					}
-					if (bCleanJson)
-					{
-						CleanJson(JsonDir);
-					}
-					if (bCleanUdn)
-					{
-						CleanUdn(Path.Combine(UdnDir, "BlueprintAPI"));
-					}
-					if (bCleanHtml)
-					{
-						CleanHtml(Path.Combine(HtmlDir, "INT\\BlueprintAPI"));
-					}
-					if (bCleanChm)
-					{
-						CleanChm(ChmDir);
-					}
-
-					// Build the data
-					if (!bBuildTargetInfo || BuildTargetInfo(TargetInfoPath, EngineDir))
-					{
-						if (!bBuildJson || BuildJson(JsonDir, EngineDir, EditorPath))
-						{
-							if (!bBuildUdn || BuildBlueprintUdn(JsonDir, UdnDir, ChmDir, StatsPath))
-							{
-								if (!bBuildHtml || BuildHtml(EngineDir, DocToolPath, UdnDir, HtmlDir, true))
-								{
-									if (!bBuildChm || BuildChm(ChmCompilerPath, HtmlDir, ChmDir, true))
-									{
-										Console.WriteLine("Complete.");
-									}
-								}
-							}
-						}
-					}
-				}
-			}
-			else
-			{
-				// Write the command line options
 				Console.WriteLine("APIDocTool.exe [options] -enginedir=<...>");                                    // <-- 80 character limit to start of comment
 				Console.WriteLine();
 				Console.WriteLine("Options:");
-				Console.WriteLine("    -blueprint:                      Build the blueprint API");
-				Console.WriteLine("    -code:                           Build the code API");
 				Console.WriteLine("    -rebuild:                        Clean and build everything");
-				Console.WriteLine("    -rebuild[meta|xml|udn|html|chm]: Clean and build specific files");
+				Console.WriteLine("    -rebuild<step>:					Clean and build specific steps");
 				Console.WriteLine("    -clean:                          Clean all files");
-				Console.WriteLine("    -clean[meta|xml|udn|html|chm]:   Clean specific files");
+				Console.WriteLine("    -clean<step>:					Clean specific steps");
 				Console.WriteLine("    -build:						    Build everything");
-				Console.WriteLine("    -build[meta|xml|udn|html|chm]:   Build specific output files");
-				Console.WriteLine("    -targetinfo=<...>:			    Specifies the build info, created by");
-				Console.WriteLine("                                     running UBT with -writetargetinfo=<...>");
-				Console.WriteLine("    -xmldir=<...>:				    Output directory for xml files");
-				Console.WriteLine("    -jsondir=<...>:				    Output directory for json files");
-				Console.WriteLine("    -metadatadir=<...>:			    Output directory for metadata files");
+				Console.WriteLine("    -build<step>:					Build specific output steps");
+				Console.WriteLine("    -archive:						Archive everything");
+				Console.WriteLine("    -archive<step>:					Archive specific output steps");
+				Console.WriteLine("    -enginedir=<...>:				Specifies the root engine directory");
+				Console.WriteLine("    -intermediatedir=<...>:			Specifies the intermediate directory");
+				Console.WriteLine("    -documentationdir=<...>:			Specifies the documentation directory");
+				Console.WriteLine("    -indexonly:						Just build index pages");
 				Console.WriteLine("    -filter=<...>,<...>:             Filter conversion, eg.");
 				Console.WriteLine("                                       Folders:  -filter=Core/Containers/...");
 				Console.WriteLine("                                       Entities: -filter=Core/TArray");
+				return 1;
+			}
+
+			// Derive all the tool paths
+			string DoxygenPath = Path.Combine(EngineDir, "Extras\\NotForLicensees\\Doxygen\\bin\\doxygen.exe");
+			string EditorPath = Path.Combine(EngineDir, "Binaries\\Win64\\UE4Editor-Cmd.exe");
+			string DocToolPath = Path.Combine(EngineDir, "Binaries\\DotNET\\UnrealDocTool.exe");
+			string ChmCompilerPath = Path.Combine(EngineDir, "Extras\\NotForLicensees\\HTML Help Workshop\\hhc.exe");
+
+			// Derive all the intermediate paths
+			string CodeSitemapDir = Path.Combine(IntermediateDir, "sitemap");
+			string BlueprintSitemapDir = Path.Combine(IntermediateDir, "blueprintsitemap");
+			string TargetInfoPath = Path.Combine(IntermediateDir, "build\\targetinfo.xml");
+			string MetadataDir = Path.Combine(IntermediateDir, "metadata");
+			string XmlDir = Path.Combine(IntermediateDir, "doxygen");
+			string JsonDir = Path.Combine(IntermediateDir, "json");
+			string MetadataPath = Path.Combine(MetadataDir, "metadata.xml");
+
+			// Derive all the output paths
+			string UdnDir = Path.Combine(DocumentationDir, "Source");
+			string HtmlDir = Path.Combine(DocumentationDir, "HTML");
+			string ChmDir = Path.Combine(DocumentationDir, "CHM");
+			string ArchiveDir = Path.Combine(DocumentationDir, "Builds");
+
+			// Read the settings file
+			Settings = IniFile.Read(Path.Combine(DocumentationDir, "Extras\\API\\API.ini"));
+			IgnoredFunctionMacros = new HashSet<string>(Settings.FindValueOrDefault("Input.IgnoredFunctionMacros", "").Split('\n'));
+			IncludedSourceFiles = new HashSet<string>(Settings.FindValueOrDefault("Output.IncludedSourceFiles", "").Split('\n'));
+
+			// Find all the metadata pages
+			AddMetadataKeyword(UdnDir, "UCLASS", "Programming/UnrealArchitecture/Reference/Classes#classdeclaration", "Programming/UnrealArchitecture/Reference/Classes/Specifiers");
+			AddMetadataKeyword(UdnDir, "UFUNCTION", "Programming/UnrealArchitecture/Reference/Functions", "Programming/UnrealArchitecture/Reference/Functions/Specifiers");
+			AddMetadataKeyword(UdnDir, "UPROPERTY", "Programming/UnrealArchitecture/Reference/Properties", "Programming/UnrealArchitecture/Reference/Properties/Specifiers");
+			AddMetadataKeyword(UdnDir, "USTRUCT", "Programming/UnrealArchitecture/Reference/Structs", "Programming/UnrealArchitecture/Reference/Structs/Specifiers");
+
+			// Build all the code docs
+			if (!BuildCodeTargetInfo(TargetInfoPath, EngineDir, Path.Combine(ArchiveDir, "CodeAPI-TargetInfo.tgz"), CodeTargetActions))
+			{
+				return 1;
 			}
+			if (!BuildCodeMetadata(DoxygenPath, EngineDir, MetadataDir, MetadataPath, Path.Combine(ArchiveDir, "CodeAPI-Metadata.tgz"), CodeMetadataActions))
+			{
+				return 1;
+			}
+			if (!BuildCodeXml(EngineDir, TargetInfoPath, DoxygenPath, XmlDir, Path.Combine(ArchiveDir, "CodeAPI-XML.tgz"), Filters, CodeXmlActions))
+			{
+				return 1;
+			}
+			if (!BuildCodeUdn(EngineDir, XmlDir, UdnDir, CodeSitemapDir, MetadataPath, Path.Combine(ArchiveDir, "CodeAPI-UDN.tgz"), Path.Combine(ArchiveDir, "CodeAPI-Sitemap.tgz"), Filters, CodeUdnActions))
+			{
+				return 1;
+			}
+			if (!BuildHtml(EngineDir, DocToolPath, HtmlDir, "API", Path.Combine(ArchiveDir, "CodeAPI-HTML.tgz"), CodeHtmlActions))
+			{
+				return 1;
+			}
+			if (!BuildChm(ChmCompilerPath, Path.Combine(ChmDir, "API.chm"), Path.Combine(CodeSitemapDir, "API.hhc"), Path.Combine(CodeSitemapDir, "API.hhk"), HtmlDir, false, CodeChmActions))
+			{
+				return 1;
+			}
+
+			// Build all the blueprint docs
+			if (!BuildBlueprintJson(JsonDir, EngineDir, EditorPath, Path.Combine(ArchiveDir, "BlueprintAPI-JSON.tgz"), BlueprintJsonActions))
+			{
+				return 1;
+			}
+			if (!BuildBlueprintUdn(JsonDir, UdnDir, BlueprintSitemapDir, Path.Combine(ArchiveDir, "BlueprintAPI-UDN.tgz"), Path.Combine(ArchiveDir, "BlueprintAPI-Sitemap.tgz"), BlueprintUdnActions))
+			{
+				return 1;
+			}
+			if (!BuildHtml(EngineDir, DocToolPath, HtmlDir, "BlueprintAPI", Path.Combine(ArchiveDir, "BlueprintAPI-HTML.tgz"), BlueprintHtmlActions))
+			{
+				return 1;
+			}
+			if (!BuildChm(ChmCompilerPath, Path.Combine(ChmDir, "BlueprintAPI.chm"), Path.Combine(BlueprintSitemapDir, "BlueprintAPI.hhc"), Path.Combine(BlueprintSitemapDir, "BlueprintAPI.hhk"), HtmlDir, true, BlueprintChmActions))
+			{
+				return 1;
+			}
+
+			// Finished!
+			Console.WriteLine("Complete.");
+			return 0;
 		}
 
+		static BuildActions ParseBuildActions(List<string> ArgumentList, string Suffix)
+		{
+			BuildActions Actions = 0;
+			if (ArgumentList.Remove("-rebuild" + Suffix))
+			{
+				Actions |= BuildActions.Clean | BuildActions.Build | BuildActions.Archive;
+			}
+			if (ArgumentList.Remove("-clean" + Suffix))
+			{
+				Actions |= BuildActions.Clean;
+			}
+			if (ArgumentList.Remove("-build" + Suffix))
+			{
+				Actions |= BuildActions.Build;
+			}
+			if (ArgumentList.Remove("-archive" + Suffix))
+			{
+				Actions |= BuildActions.Archive;
+			}
+			return Actions;
+		}
+
+		static string ParseArgumentValue(List<string> ArgumentList, string Prefix, string DefaultValue)
+		{
+			for (int Idx = 0; Idx < ArgumentList.Count; Idx++)
+			{
+				if (ArgumentList[Idx].StartsWith(Prefix))
+				{
+					string Value = ArgumentList[Idx].Substring(Prefix.Length);
+					ArgumentList.RemoveAt(Idx);
+					return Value;
+				}
+			}
+			return DefaultValue;
+		}
+
+		static string ParseArgumentPath(List<string> ArgumentList, string Prefix, string DefaultValue)
+		{
+			string Value = ParseArgumentValue(ArgumentList, Prefix, DefaultValue);
+			if (Value != null)
+			{
+				Value = Path.GetFullPath(Value);
+			}
+			return Value;
+		}
+
+		static string ParseArgumentDirectory(List<string> ArgumentList, string Prefix, string DefaultValue)
+		{
+			string Value = ParseArgumentPath(ArgumentList, Prefix, DefaultValue);
+			if(Value != null && !Directory.Exists(Value))
+			{
+				Directory.CreateDirectory(Value);
+			}
+			return Value;
+		}
+
+
 		static void AddMetadataKeyword(string BaseDir, string Name, string TypeUrl, string SpecifierBaseUrl)
 		{
 			MetadataKeyword Keyword = new MetadataKeyword();
@@ -545,31 +422,38 @@ namespace APIDocTool
 			MetadataDirective.AllKeywords.Add(Name, Keyword);
 		}
 
-		static void CleanTargetInfo(string TargetInfoPath)
-		{
-			Console.WriteLine("Cleaning '{0}'", TargetInfoPath);
-			Utility.SafeDeleteFile(TargetInfoPath);
-		}
-
-		static bool BuildTargetInfo(string TargetInfoPath, string EngineDir)
+		static bool BuildCodeTargetInfo(string TargetInfoPath, string EngineDir, string ArchivePath, BuildActions Actions)
 		{
-			Console.WriteLine("Building target info...");
-			Utility.SafeCreateDirectory(Path.GetDirectoryName(TargetInfoPath));
-
-			string Arguments = String.Format("DocumentationEditor Win64 Debug -project=\"{0}\"", Path.Combine(EngineDir, "Documentation\\Extras\\API\\Build\\Documentation.uproject"));
-			if(!RunUnrealBuildTool(EngineDir, Arguments + " -clean"))
+			if((Actions & BuildActions.Clean) != 0)
 			{
-				return false;
+				Console.WriteLine("Cleaning '{0}'", TargetInfoPath);
+				Utility.SafeDeleteFile(TargetInfoPath);
 			}
-			foreach(FileInfo Info in new DirectoryInfo(Path.Combine(EngineDir, "Intermediate\\Build")).EnumerateFiles("UBTEXport*.xml"))
+			if ((Actions & BuildActions.Build) != 0)
 			{
-				File.Delete(Info.FullName);
+				Console.WriteLine("Building target info...");
+				Utility.SafeCreateDirectory(Path.GetDirectoryName(TargetInfoPath));
+
+				string Arguments = String.Format("DocumentationEditor Win64 Debug -project=\"{0}\"", Path.Combine(EngineDir, "Documentation\\Extras\\API\\Build\\Documentation.uproject"));
+				if (!RunUnrealBuildTool(EngineDir, Arguments + " -clean"))
+				{
+					return false;
+				}
+				foreach (FileInfo Info in new DirectoryInfo(Path.Combine(EngineDir, "Intermediate\\Build")).EnumerateFiles("UBTEXport*.xml"))
+				{
+					File.Delete(Info.FullName);
+				}
+				if (!RunUnrealBuildTool(EngineDir, Arguments + " -disableunity -xgeexport"))
+				{
+					return false;
+				}
+				File.Copy(Path.Combine(EngineDir, "Intermediate\\Build\\UBTExport.0.xge.xml"), TargetInfoPath, true);
 			}
-			if(!RunUnrealBuildTool(EngineDir, Arguments + " -disableunity -xgeexport"))
+			if((Actions & BuildActions.Archive) != 0)
 			{
-				return false;
+				Console.WriteLine("Creating archive '{0}'", ArchivePath);
+				Utility.CreateTgzFromDir(ArchivePath, Path.GetDirectoryName(TargetInfoPath));
 			}
-			File.Copy(Path.Combine(EngineDir, "Intermediate\\Build\\UBTExport.0.xge.xml"), TargetInfoPath, true);
 			return true;
 		}
 
@@ -604,21 +488,24 @@ namespace APIDocTool
 			}
 		}
 
-
-		static void CleanMetadata(string MetadataDir)
+		static bool BuildCodeMetadata(string DoxygenPath, string EngineDir, string MetadataDir, string MetadataPath, string ArchivePath, BuildActions Actions)
 		{
-			Console.WriteLine("Cleaning '{0}'", MetadataDir);
-			Utility.SafeDeleteDirectoryContents(MetadataDir, true);
-		}
+			if((Actions & BuildActions.Clean) != 0)
+			{
+				Console.WriteLine("Cleaning '{0}'", MetadataDir);
+				Utility.SafeDeleteDirectoryContents(MetadataDir, true);
+			}
+			if((Actions & BuildActions.Build) != 0)
+			{
+				string MetadataInputPath = Path.Combine(EngineDir, "Source\\Runtime\\CoreUObject\\Public\\UObject\\ObjectBase.h");
+				Console.WriteLine("Building metadata descriptions from '{0}'...", MetadataInputPath);
 
-		static bool BuildMetadata(string DoxygenPath, string EngineDir, string MetadataDir, string MetadataPath)
-		{
-			string MetadataInputPath = Path.Combine(EngineDir, "Source\\Runtime\\CoreUObject\\Public\\UObject\\ObjectBase.h");
-			Console.WriteLine("Building metadata descriptions from '{0}'...", MetadataInputPath);
+				DoxygenConfig Config = new DoxygenConfig("Metadata", new string[]{ MetadataInputPath }, MetadataDir);
+				if (!Doxygen.Run(DoxygenPath, Path.Combine(EngineDir, "Source"), Config, true))
+				{
+					return false;
+				}
 
-			DoxygenConfig Config = new DoxygenConfig("Metadata", new string[]{ MetadataInputPath }, MetadataDir);
-			if (Doxygen.Run(DoxygenPath, Path.Combine(EngineDir, "Source"), Config, true))
-			{
 				MetadataLookup.Reset();
 
 				// Parse the xml output
@@ -634,9 +521,13 @@ namespace APIDocTool
 
 				// Write it to a summary file
 				MetadataLookup.Save(MetadataPath);
-				return true;
 			}
-			return false;
+			if ((Actions & BuildActions.Archive) != 0)
+			{
+				Console.WriteLine("Creating archive '{0}'", ArchivePath);
+				Utility.CreateTgzFromDir(ArchivePath, MetadataDir);
+			}
+			return true;
 		}
 
 		static bool ParseMetadataTags(string InputFile, Dictionary<string, string> Tags)
@@ -668,352 +559,385 @@ namespace APIDocTool
 			return false;
 		}
 
-		static void CleanXml(string XmlDir)
+		static bool BuildCodeXml(string EngineDir, string TargetInfoPath, string DoxygenPath, string XmlDir, string ArchivePath, List<string> Filters, BuildActions Actions)
 		{
-			Console.WriteLine("Cleaning '{0}'", XmlDir);
-			Utility.SafeDeleteDirectoryContents(XmlDir, true);
-		}
-
-		static bool BuildXml(string EngineDir, string TargetInfoPath, string DoxygenPath, string XmlDir, List<string> Filters = null)
-		{
-			// Create the output directory
-			Utility.SafeCreateDirectory(XmlDir);
+			if ((Actions & BuildActions.Clean) != 0)
+			{
+				Console.WriteLine("Cleaning '{0}'", XmlDir);
+				Utility.SafeDeleteDirectoryContents(XmlDir, true);
+			}
+			if ((Actions & BuildActions.Build) != 0)
+			{
+				Console.WriteLine("Building XML...");
+				Utility.SafeCreateDirectory(XmlDir);
 
-			// Read the target that we're building
-			BuildTarget Target = new BuildTarget(Path.Combine(EngineDir, "Source"), TargetInfoPath);
+				// Read the target that we're building
+				BuildTarget Target = new BuildTarget(Path.Combine(EngineDir, "Source"), TargetInfoPath);
 
-			// Create an invariant list of exclude directories
-			string[] InvariantExcludeDirectories = ExcludeSourceDirectories.Select(x => x.ToLowerInvariant()).ToArray();
+				// Create an invariant list of exclude directories
+				string[] InvariantExcludeDirectories = ExcludeSourceDirectories.Select(x => x.ToLowerInvariant()).ToArray();
 
-			// Get the list of folders to filter against
-			List<string> FolderFilters = new List<string>();
-			if(Filters != null)
-			{
-				foreach(string Filter in Filters)
+				// Get the list of folders to filter against
+				List<string> FolderFilters = new List<string>();
+				if (Filters != null)
 				{
-					int Idx = Filter.IndexOf('/');
-					if(Idx != -1)
+					foreach (string Filter in Filters)
 					{
-						FolderFilters.Add("\\" + Filter.Substring(0, Idx) + "\\");
+						int Idx = Filter.IndexOf('/');
+						if (Idx != -1)
+						{
+							FolderFilters.Add("\\" + Filter.Substring(0, Idx) + "\\");
+						}
 					}
 				}
-			}
 
-			// Flatten the target into a list of modules
-			List<string> InputModules = new List<string>();
-			foreach (string DirectoryName in Target.DirectoryNames)
-			{
-				for(DirectoryInfo ModuleDirectory = new DirectoryInfo(DirectoryName); ModuleDirectory.Parent != null; ModuleDirectory = ModuleDirectory.Parent)
+				// Flatten the target into a list of modules
+				List<string> InputModules = new List<string>();
+				foreach (string DirectoryName in Target.DirectoryNames)
 				{
-					IEnumerator<FileInfo> ModuleFile = ModuleDirectory.EnumerateFiles("*.build.cs").GetEnumerator();
-					if(ModuleFile.MoveNext() && (FolderFilters.Count == 0 || FolderFilters.Any(x => ModuleFile.Current.FullName.Contains(x))))
+					for (DirectoryInfo ModuleDirectory = new DirectoryInfo(DirectoryName); ModuleDirectory.Parent != null; ModuleDirectory = ModuleDirectory.Parent)
 					{
-						InputModules.AddUnique(ModuleFile.Current.FullName);
-						break;
+						IEnumerator<FileInfo> ModuleFile = ModuleDirectory.EnumerateFiles("*.build.cs").GetEnumerator();
+						if (ModuleFile.MoveNext() && (FolderFilters.Count == 0 || FolderFilters.Any(x => ModuleFile.Current.FullName.Contains(x))))
+						{
+							InputModules.AddUnique(ModuleFile.Current.FullName);
+							break;
+						}
 					}
 				}
-			}
 
-			// Just use all the input modules
-			if(!bIndexOnly)
-			{
-				// Set our error mode so as to not bring up the WER dialog if Doxygen crashes (which it often does)
-				SetErrorMode(0x0007);
+				// Just use all the input modules
+				if (!bIndexOnly)
+				{
+					// Set our error mode so as to not bring up the WER dialog if Doxygen crashes (which it often does)
+					SetErrorMode(0x0007);
 
-				// Create the output directory
-				Utility.SafeCreateDirectory(XmlDir);
+					// Create the output directory
+					Utility.SafeCreateDirectory(XmlDir);
 
-				// Build all the modules
-				Console.WriteLine("Parsing source...");
+					// Build all the modules
+					Console.WriteLine("Parsing source...");
 
-				// Build the list of definitions
-				List<string> Definitions = new List<string>();
-				foreach(string Definition in Target.Definitions)
-				{
-					if(!Definition.StartsWith("ORIGINAL_FILE_NAME="))
+					// Build the list of definitions
+					List<string> Definitions = new List<string>();
+					foreach (string Definition in Target.Definitions)
 					{
-						Definitions.Add(Definition.Replace("DLLIMPORT", "").Replace("DLLEXPORT", ""));
+						if (!Definition.StartsWith("ORIGINAL_FILE_NAME="))
+						{
+							Definitions.Add(Definition.Replace("DLLIMPORT", "").Replace("DLLEXPORT", ""));
+						}
 					}
-				}
 
-				// Build a list of input paths
-				List<string> InputPaths = new List<string>();
-				foreach(string InputModule in InputModules)
-				{
-					foreach(string DirectoryName in Directory.EnumerateDirectories(Path.GetDirectoryName(InputModule), "*", SearchOption.AllDirectories))
+					// Build a list of input paths
+					List<string> InputPaths = new List<string>();
+					foreach (string InputModule in InputModules)
 					{
-						// Find the relative path from the engine directory
-						string NormalizedDirectoryName = DirectoryName;
-						if(NormalizedDirectoryName.StartsWith(EngineDir))
+						foreach (string DirectoryName in Directory.EnumerateDirectories(Path.GetDirectoryName(InputModule), "*", SearchOption.AllDirectories))
 						{
-							NormalizedDirectoryName = NormalizedDirectoryName.Substring(EngineDir.Length);
-						}
-						if(!NormalizedDirectoryName.EndsWith("\\"))
-						{
-							NormalizedDirectoryName += "\\";
-						}
+							// Find the relative path from the engine directory
+							string NormalizedDirectoryName = DirectoryName;
+							if (NormalizedDirectoryName.StartsWith(EngineDir))
+							{
+								NormalizedDirectoryName = NormalizedDirectoryName.Substring(EngineDir.Length);
+							}
+							if (!NormalizedDirectoryName.EndsWith("\\"))
+							{
+								NormalizedDirectoryName += "\\";
+							}
 
-						// Check we can include it
-						if(!ExcludeSourceDirectories.Any(x => NormalizedDirectoryName.Contains("\\" + x + "\\")))
-						{
-							if(FolderFilters.Count == 0 || FolderFilters.Any(x => NormalizedDirectoryName.Contains(x)))
+							// Check we can include it
+							if (!ExcludeSourceDirectories.Any(x => NormalizedDirectoryName.Contains("\\" + x + "\\")))
 							{
-								InputPaths.Add(DirectoryName);
+								if (FolderFilters.Count == 0 || FolderFilters.Any(x => NormalizedDirectoryName.Contains(x)))
+								{
+									InputPaths.Add(DirectoryName);
+								}
 							}
 						}
 					}
-				}
-
-				// Build the configuration for this module
-				DoxygenConfig Config = new DoxygenConfig("UE4", InputPaths.ToArray(), XmlDir);
-				Config.Definitions.AddRange(Definitions);
-				Config.Definitions.AddRange(DoxygenPredefinedMacros);
-				Config.ExpandAsDefined.AddRange(DoxygenExpandedMacros);
-				Config.IncludePaths.AddRange(Target.IncludePaths);
-				Config.ExcludePatterns.AddRange(ExcludeSourceDirectories.Select(x => "*/" + x + "/*"));
-				Config.ExcludePatterns.AddRange(ExcludeSourceFiles);
 
-				// Run Doxygen
-				if (!Doxygen.Run(DoxygenPath, Path.Combine(EngineDir, "Source"), Config, true))
-				{
-					Console.WriteLine("  Doxygen crashed. Skipping.");
-					return false;
+					// Build the configuration for this module
+					DoxygenConfig Config = new DoxygenConfig("UE4", InputPaths.ToArray(), XmlDir);
+					Config.Definitions.AddRange(Definitions);
+					Config.Definitions.AddRange(DoxygenPredefinedMacros);
+					Config.ExpandAsDefined.AddRange(DoxygenExpandedMacros);
+					Config.IncludePaths.AddRange(Target.IncludePaths);
+					Config.ExcludePatterns.AddRange(ExcludeSourceDirectories.Select(x => "*/" + x + "/*"));
+					Config.ExcludePatterns.AddRange(ExcludeSourceFiles);
+
+					// Run Doxygen
+					if (!Doxygen.Run(DoxygenPath, Path.Combine(EngineDir, "Source"), Config, true))
+					{
+						Console.WriteLine("  Doxygen crashed. Skipping.");
+						return false;
+					}
 				}
-			}
 
-			// Write the modules file
-			File.WriteAllLines(Path.Combine(XmlDir, "modules.txt"), InputModules);
+				// Write the modules file
+				File.WriteAllLines(Path.Combine(XmlDir, "modules.txt"), InputModules);
+			}
+			if((Actions & BuildActions.Archive) != 0)
+			{
+				Console.WriteLine("Creating archive '{0}'", ArchivePath);
+				Utility.CreateTgzFromDir(ArchivePath, XmlDir);
+			}
 			return true;
 		}
 
-		static void CleanJson(string JsonDir)
+		static bool BuildBlueprintJson(string JsonDir, string EngineDir, string EditorPath, string ArchivePath, BuildActions Actions)
 		{
-			Console.WriteLine("Cleaning '{0}'", JsonDir);
-			Utility.SafeDeleteDirectoryContents(JsonDir, true);
-		}
-
-		static bool BuildJson(string JsonDir, string EngineDir, string EditorPath)
-		{
-			// Create the output directory
-			Utility.SafeCreateDirectory(JsonDir);
-
-			using (Process JsonExportProcess = new Process())
+			if((Actions & BuildActions.Clean) != 0)
 			{
-				JsonExportProcess.StartInfo.WorkingDirectory = EngineDir;
-				JsonExportProcess.StartInfo.FileName = EditorPath;
-				JsonExportProcess.StartInfo.Arguments = "-run=GenerateBlueprintAPI -path=" + JsonDir + " -name=BlueprintAPI";
-				JsonExportProcess.StartInfo.UseShellExecute = false;
-				JsonExportProcess.StartInfo.RedirectStandardOutput = true;
-				JsonExportProcess.StartInfo.RedirectStandardError = true;
-
-				JsonExportProcess.OutputDataReceived += new DataReceivedEventHandler(ProcessOutputReceived);
-				JsonExportProcess.ErrorDataReceived += new DataReceivedEventHandler(ProcessOutputReceived);
+				Console.WriteLine("Cleaning '{0}'", JsonDir);
+				Utility.SafeDeleteDirectoryContents(JsonDir, true);
+			}
+			if ((Actions & BuildActions.Build) != 0)
+			{
+				// Create the output directory
+				Utility.SafeCreateDirectory(JsonDir);
 
-				try
-				{
-					JsonExportProcess.Start();
-					JsonExportProcess.BeginOutputReadLine();
-					JsonExportProcess.BeginErrorReadLine();
-					JsonExportProcess.WaitForExit();
-					return JsonExportProcess.ExitCode == 0;
-				}
-				catch (Exception Ex)
+				using (Process JsonExportProcess = new Process())
 				{
-					Console.WriteLine(Ex.ToString() + "\n" + Ex.StackTrace);
-					return false;
-				}
-			}
-		}
+					JsonExportProcess.StartInfo.WorkingDirectory = EngineDir;
+					JsonExportProcess.StartInfo.FileName = EditorPath;
+					JsonExportProcess.StartInfo.Arguments = "-run=GenerateBlueprintAPI -path=" + JsonDir + " -name=BlueprintAPI";
+					JsonExportProcess.StartInfo.UseShellExecute = false;
+					JsonExportProcess.StartInfo.RedirectStandardOutput = true;
+					JsonExportProcess.StartInfo.RedirectStandardError = true;
 
-		static void CleanUdn(string CleanDir)
-		{
-			Console.WriteLine("Cleaning '{0}'", CleanDir);
+					JsonExportProcess.OutputDataReceived += new DataReceivedEventHandler(ProcessOutputReceived);
+					JsonExportProcess.ErrorDataReceived += new DataReceivedEventHandler(ProcessOutputReceived);
 
-			// Delete all the files
-			foreach(string FileName in Directory.EnumerateFiles(CleanDir))
-			{
-				Utility.SafeDeleteFile(Path.Combine(CleanDir, FileName));
+					try
+					{
+						JsonExportProcess.Start();
+						JsonExportProcess.BeginOutputReadLine();
+						JsonExportProcess.BeginErrorReadLine();
+						JsonExportProcess.WaitForExit();
+						return JsonExportProcess.ExitCode == 0;
+					}
+					catch (Exception Ex)
+					{
+						Console.WriteLine(Ex.ToString() + "\n" + Ex.StackTrace);
+						return false;
+					}
+				}
 			}
-
-			// Delete all the subdirectories
-			foreach (string SubDir in Directory.EnumerateDirectories(CleanDir))
+			if((Actions & BuildActions.Archive) != 0)
 			{
-				Utility.SafeDeleteDirectory(SubDir);
+				Console.WriteLine("Creating archive '{0}'", ArchivePath);
+				Utility.CreateTgzFromDir(ArchivePath, JsonDir);
 			}
+			return true;
 		}
 
-		static bool BuildUdn(string EngineDir, string XmlDir, string UdnDir, string SitemapDir, string MetadataPath, string StatsPath, List<string> Filters = null)
+		static bool BuildCodeUdn(string EngineDir, string XmlDir, string UdnDir, string SitemapDir, string MetadataPath, string ArchivePath, string SitemapArchivePath, List<string> Filters, BuildActions Actions)
 		{
-			// Create the output directory
-			Utility.SafeCreateDirectory(UdnDir);
-
-			// Read the metadata
-			MetadataLookup.Load(MetadataPath);
-
-			// Read the list of modules
-			List<string> InputModules = new List<string>(File.ReadAllLines(Path.Combine(XmlDir, "modules.txt")));
-
-			// Build the doxygen modules
-			List<DoxygenModule> Modules = new List<DoxygenModule>();
-			foreach(string InputModule in InputModules)
+			string ApiDir = Path.Combine(UdnDir, "API");
+			if((Actions & BuildActions.Clean) != 0)
 			{
-				Modules.Add(new DoxygenModule(Path.GetFileNameWithoutExtension(Path.GetFileNameWithoutExtension(InputModule)), Path.GetDirectoryName(InputModule)));
+				Console.WriteLine("Cleaning '{0}'", ApiDir);
+				Utility.SafeDeleteDirectoryContents(ApiDir, true);
+				Utility.SafeDeleteDirectoryContents(SitemapDir, true);
 			}
-
-			// Find all the entities
-			if(!bIndexOnly)
+			if ((Actions & BuildActions.Build) != 0)
 			{
-				Console.WriteLine("Reading Doxygen output...");
+				Directory.CreateDirectory(ApiDir);
+				Directory.CreateDirectory(SitemapDir);
 
-				// Read the engine module and split it into smaller modules
-				DoxygenModule RootModule = DoxygenModule.Read("UE4", EngineDir, Path.Combine(XmlDir, "xml"));
-				foreach(DoxygenEntity Entity in RootModule.Entities)
-				{
-					DoxygenModule Module = Modules.Find(x => Entity.File.StartsWith(x.BaseSrcDir));
-					Entity.Module = Module;
-					Module.Entities.Add(Entity);
-				}
-				foreach(DoxygenSourceFile SourceFile in RootModule.SourceFiles)
+				// Read the metadata
+				MetadataLookup.Load(MetadataPath);
+
+				// Read the list of modules
+				List<string> InputModules = new List<string>(File.ReadAllLines(Path.Combine(XmlDir, "modules.txt")));
+
+				// Build the doxygen modules
+				List<DoxygenModule> Modules = new List<DoxygenModule>();
+				foreach (string InputModule in InputModules)
 				{
-					DoxygenModule Module = Modules.Find(x => SourceFile.FileName.Replace('/', '\\').StartsWith(x.BaseSrcDir));
-					Module.SourceFiles.Add(SourceFile);
+					Modules.Add(new DoxygenModule(Path.GetFileNameWithoutExtension(Path.GetFileNameWithoutExtension(InputModule)), Path.GetDirectoryName(InputModule)));
 				}
 
-				// Now filter all the entities in each module
-				if(Filters != null && Filters.Count > 0)
+				// Find all the entities
+				if (!bIndexOnly)
 				{
-					FilterEntities(Modules, Filters);
-				}
+					Console.WriteLine("Reading Doxygen output...");
 
-				// Remove all the empty modules
-				Modules.RemoveAll(x => x.Entities.Count == 0);
-			}
+					// Read the engine module and split it into smaller modules
+					DoxygenModule RootModule = DoxygenModule.Read("UE4", EngineDir, Path.Combine(XmlDir, "xml"));
+					foreach (DoxygenEntity Entity in RootModule.Entities)
+					{
+						DoxygenModule Module = Modules.Find(x => Entity.File.StartsWith(x.BaseSrcDir));
+						Entity.Module = Module;
+						Module.Entities.Add(Entity);
+					}
+					foreach (DoxygenSourceFile SourceFile in RootModule.SourceFiles)
+					{
+						DoxygenModule Module = Modules.Find(x => SourceFile.FileName.Replace('/', '\\').StartsWith(x.BaseSrcDir));
+						Module.SourceFiles.Add(SourceFile);
+					}
 
-			// Create the index page, and all the pages below it
-			APIIndex Index = new APIIndex(Modules);
+					// Now filter all the entities in each module
+					if (Filters != null && Filters.Count > 0)
+					{
+						FilterEntities(Modules, Filters);
+					}
 
-			// Build a list of pages to output
-			List<APIPage> OutputPages = new List<APIPage>(Index.GatherPages().OrderBy(x => x.LinkPath));
+					// Remove all the empty modules
+					Modules.RemoveAll(x => x.Entities.Count == 0);
+				}
 
-			// Dump the output stats
-			if (StatsPath != null)
-			{
-				Console.WriteLine("Writing stats to " + StatsPath + "...");
+				// Create the index page, and all the pages below it
+				APIIndex Index = new APIIndex(Modules);
+
+				// Build a list of pages to output
+				List<APIPage> OutputPages = new List<APIPage>(Index.GatherPages().OrderBy(x => x.LinkPath));
+
+				// Dump the output stats
+				string StatsPath = Path.Combine(SitemapDir, "Stats.txt");
+				Console.WriteLine("Writing stats to '" + StatsPath + "'");
 				Stats NewStats = new Stats(OutputPages.OfType<APIMember>());
 				NewStats.Write(StatsPath);
-			}
 
-			// Setup the output directory 
-			Utility.SafeCreateDirectory(UdnDir);
+				// Setup the output directory 
+				Utility.SafeCreateDirectory(UdnDir);
 
-			// Build the manifest
-			Console.WriteLine("Writing manifest...");
-			UdnManifest Manifest = new UdnManifest(Index);
-			Manifest.PrintConflicts();
-			Manifest.Write(Path.Combine(UdnDir, APIFolder + "\\API.manifest"));
+				// Build the manifest
+				Console.WriteLine("Writing manifest...");
+				UdnManifest Manifest = new UdnManifest(Index);
+				Manifest.PrintConflicts();
+				Manifest.Write(Path.Combine(UdnDir, APIFolder + "\\API.manifest"));
 
-			// Write all the pages
-			using (Tracker UdnTracker = new Tracker("Writing UDN pages...", OutputPages.Count))
-			{
-				foreach(int Idx in UdnTracker.Indices)
+				// Write all the pages
+				using (Tracker UdnTracker = new Tracker("Writing UDN pages...", OutputPages.Count))
 				{
-					APIPage Page = OutputPages[Idx];
-
-					// Create the output directory
-					string MemberDirectory = Path.Combine(UdnDir, Page.LinkPath);
-					if (!Directory.Exists(MemberDirectory))
+					foreach (int Idx in UdnTracker.Indices)
 					{
-						Directory.CreateDirectory(MemberDirectory);
-					}
+						APIPage Page = OutputPages[Idx];
+
+						// Create the output directory
+						string MemberDirectory = Path.Combine(UdnDir, Page.LinkPath);
+						if (!Directory.Exists(MemberDirectory))
+						{
+							Directory.CreateDirectory(MemberDirectory);
+						}
 
-					// Write the page
-					Page.WritePage(Manifest, Path.Combine(MemberDirectory, "index.INT.udn"));
+						// Write the page
+						Page.WritePage(Manifest, Path.Combine(MemberDirectory, "index.INT.udn"));
+					}
 				}
-			}
 
-			// Write the sitemap contents
-			Console.WriteLine("Writing sitemap contents...");
-			Index.WriteSitemapContents(Path.Combine(SitemapDir, "API.hhc"));
+				// Write the sitemap contents
+				Console.WriteLine("Writing sitemap contents...");
+				Index.WriteSitemapContents(Path.Combine(SitemapDir, "API.hhc"));
 
-			// Write the sitemap index
-			Console.WriteLine("Writing sitemap index...");
-			Index.WriteSitemapIndex(Path.Combine(SitemapDir, "API.hhk"));
+				// Write the sitemap index
+				Console.WriteLine("Writing sitemap index...");
+				Index.WriteSitemapIndex(Path.Combine(SitemapDir, "API.hhk"));
+			}
+			if ((Actions & BuildActions.Archive) != 0)
+			{
+				Console.WriteLine("Creating archive '{0}'", ArchivePath);
+				Utility.CreateTgzFromDir(ArchivePath, ApiDir);
 
+				Console.WriteLine("Creating archive '{0}'", SitemapArchivePath);
+				Utility.CreateTgzFromDir(SitemapArchivePath, SitemapDir);
+			}
 			return true;
 		}
 
-		static bool BuildBlueprintUdn(string JsonDir, string UdnDir, string SitemapDir, string StatsPath)
+		static bool BuildBlueprintUdn(string JsonDir, string UdnDir, string SitemapDir, string ArchivePath, string SitemapArchivePath, BuildActions ExecActions)
 		{
-			// Read the input json file
-			string JsonFilePath = Path.Combine(JsonDir, "BlueprintAPI.json");
-			var json = (Dictionary<string, object>)fastJSON.JSON.Instance.Parse(File.ReadAllText(JsonFilePath));
-
-			APICategory.LoadTooltips((Dictionary<string, object>)json["Categories"]);
+			string ApiDir = Path.Combine(UdnDir, "BlueprintAPI");
+			if ((ExecActions & BuildActions.Clean) != 0)
+			{
+				Console.WriteLine("Cleaning '{0}'", ApiDir);
+				Utility.SafeDeleteDirectoryContents(ApiDir, true);
+				Utility.SafeDeleteDirectoryContents(SitemapDir, true);
+			}
+			if ((ExecActions & BuildActions.Build) != 0)
+			{
+				Directory.CreateDirectory(ApiDir);
+				Directory.CreateDirectory(SitemapDir);
 
-			// TODO: This path is clearly sketchy as hell, but we'll clean it up later maybe
-			var Actions = (Dictionary<string, object>)((Dictionary<string, object>)((Dictionary<string, object>)((Dictionary<string, object>)json["Actor"])["Palette"])["ActionSet"])["Actions"];
+				// Read the input json file
+				string JsonFilePath = Path.Combine(JsonDir, "BlueprintAPI.json");
+				var json = (Dictionary<string, object>)fastJSON.JSON.Instance.Parse(File.ReadAllText(JsonFilePath));
 
-			APICategory RootCategory = new APICategory(null, "BlueprintAPI", true);
+				APICategory.LoadTooltips((Dictionary<string, object>)json["Categories"]);
 
-			foreach (var Action in Actions)
-			{
-				var CategoryList = Action.Key.Split('|');
-				var ActionCategory = RootCategory;
+				// TODO: This path is clearly sketchy as hell, but we'll clean it up later maybe
+				var Actions = (Dictionary<string, object>)((Dictionary<string, object>)((Dictionary<string, object>)((Dictionary<string, object>)json["Actor"])["Palette"])["ActionSet"])["Actions"];
 
-				Debug.Assert(CategoryList.Length > 0);
-				Debug.Assert(CategoryList[0] == "Library");
+				APICategory RootCategory = new APICategory(null, "BlueprintAPI", true);
 
-				for (int CategoryIndex = 1; CategoryIndex < CategoryList.Length - 1; ++CategoryIndex)
+				foreach (var Action in Actions)
 				{
-					ActionCategory = ActionCategory.GetSubCategory(CategoryList[CategoryIndex]);
-				}
+					var CategoryList = Action.Key.Split('|');
+					var ActionCategory = RootCategory;
 
-				ActionCategory.AddAction(new APIAction(ActionCategory, CategoryList.Last(), (Dictionary<string, object>)Action.Value));
-			}
+					Debug.Assert(CategoryList.Length > 0);
+					Debug.Assert(CategoryList[0] == "Library");
 
-			// Build a list of pages to output
-			List<APIPage> OutputPages = new List<APIPage>(RootCategory.GatherPages().OrderBy(x => x.LinkPath));
+					for (int CategoryIndex = 1; CategoryIndex < CategoryList.Length - 1; ++CategoryIndex)
+					{
+						ActionCategory = ActionCategory.GetSubCategory(CategoryList[CategoryIndex]);
+					}
 
-			// Create the output directory
-			Utility.SafeCreateDirectory(UdnDir);
-			Utility.SafeCreateDirectory(Path.Combine(UdnDir, APIFolder));
+					ActionCategory.AddAction(new APIAction(ActionCategory, CategoryList.Last(), (Dictionary<string, object>)Action.Value));
+				}
 
-			// Build the manifest
-			Console.WriteLine("Writing manifest...");
-			UdnManifest Manifest = new UdnManifest(RootCategory);
-			Manifest.PrintConflicts();
-			Manifest.Write(Path.Combine(UdnDir, BlueprintAPIFolder + "\\API.manifest"));
+				// Build a list of pages to output
+				List<APIPage> OutputPages = new List<APIPage>(RootCategory.GatherPages().OrderBy(x => x.LinkPath));
 
-			Console.WriteLine("Categories: " + OutputPages.Count(page => page is APICategory));
-			Console.WriteLine("Actions: " + OutputPages.Count(page => page is APIAction));
+				// Create the output directory
+				Utility.SafeCreateDirectory(UdnDir);
+				Utility.SafeCreateDirectory(Path.Combine(UdnDir, APIFolder));
 
-			// Write all the pages
-			using (Tracker UdnTracker = new Tracker("Writing UDN pages...", OutputPages.Count))
-			{
-				foreach (int Idx in UdnTracker.Indices)
-				{
-					APIPage Page = OutputPages[Idx];
+				// Build the manifest
+				Console.WriteLine("Writing manifest...");
+				UdnManifest Manifest = new UdnManifest(RootCategory);
+				Manifest.PrintConflicts();
+				Manifest.Write(Path.Combine(UdnDir, BlueprintAPIFolder + "\\API.manifest"));
 
-					// Create the output directory
-					string MemberDirectory = Path.Combine(UdnDir, Page.LinkPath);
-					if (!Directory.Exists(MemberDirectory))
+				Console.WriteLine("Categories: " + OutputPages.Count(page => page is APICategory));
+				Console.WriteLine("Actions: " + OutputPages.Count(page => page is APIAction));
+
+				// Write all the pages
+				using (Tracker UdnTracker = new Tracker("Writing UDN pages...", OutputPages.Count))
+				{
+					foreach (int Idx in UdnTracker.Indices)
 					{
-						Directory.CreateDirectory(MemberDirectory);
-					}
+						APIPage Page = OutputPages[Idx];
 
-					// Write the page
-					Page.WritePage(Manifest, Path.Combine(MemberDirectory, "index.INT.udn"));
+						// Create the output directory
+						string MemberDirectory = Path.Combine(UdnDir, Page.LinkPath);
+						if (!Directory.Exists(MemberDirectory))
+						{
+							Directory.CreateDirectory(MemberDirectory);
+						}
+
+						// Write the page
+						Page.WritePage(Manifest, Path.Combine(MemberDirectory, "index.INT.udn"));
+					}
 				}
-			}
 
-			// Write the sitemap contents
-			Console.WriteLine("Writing sitemap contents...");
-			RootCategory.WriteSitemapContents(Path.Combine(SitemapDir, "BlueprintAPI.hhc"));
+				// Write the sitemap contents
+				Console.WriteLine("Writing sitemap contents...");
+				RootCategory.WriteSitemapContents(Path.Combine(SitemapDir, "BlueprintAPI.hhc"));
 
-			// Write the sitemap index
-			Console.WriteLine("Writing sitemap index...");
-			RootCategory.WriteSitemapIndex(Path.Combine(SitemapDir, "BlueprintAPI.hhk"));
+				// Write the sitemap index
+				Console.WriteLine("Writing sitemap index...");
+				RootCategory.WriteSitemapIndex(Path.Combine(SitemapDir, "BlueprintAPI.hhk"));
+			}
+			if ((ExecActions & BuildActions.Archive) != 0)
+			{
+				Console.WriteLine("Creating archive '{0}'", ArchivePath);
+				Utility.CreateTgzFromDir(ArchivePath, ApiDir);
 
+				Console.WriteLine("Creating archive '{0}'", SitemapArchivePath);
+				Utility.CreateTgzFromDir(SitemapArchivePath, SitemapDir);
+			}
 			return true;
 		}
 
@@ -1111,164 +1035,167 @@ namespace APIDocTool
 			}
 		}
 
-		public static void CleanHtml(string CleanDir)
+		static bool BuildHtml(string EngineDir, string DocToolPath, string HtmlDir, string HtmlSuffix, string ArchivePath, BuildActions Actions)
 		{
-			Console.WriteLine("Cleaning '{0}'", CleanDir);
-			Utility.SafeDeleteDirectoryContents(CleanDir, true);
-		}
+			string HtmlSubDir = Path.Combine(HtmlDir, "INT", HtmlSuffix);
+			if((Actions & BuildActions.Clean) != 0)
+			{
+				Console.WriteLine("Cleaning '{0}'", HtmlSubDir);
+				Utility.SafeDeleteDirectoryContents(HtmlSubDir, true);
+			}
+			if ((Actions & BuildActions.Build) != 0)
+			{
+				Utility.SafeCreateDirectory(HtmlSubDir);
 
-		public static bool BuildHtml(string EngineDir, string DocToolPath, string UdnPath, string HtmlDir, bool bBlueprint)
-		{
-			Utility.SafeCreateDirectory(HtmlDir);
+				using (Process DocToolProcess = new Process())
+				{
+					DocToolProcess.StartInfo.WorkingDirectory = EngineDir;
+					DocToolProcess.StartInfo.FileName = DocToolPath;
+					DocToolProcess.StartInfo.Arguments = Path.Combine(HtmlSuffix, "*") + " -lang=INT -t=DefaultAPI.html -v=warn";
+					DocToolProcess.StartInfo.UseShellExecute = false;
+					DocToolProcess.StartInfo.RedirectStandardOutput = true;
+					DocToolProcess.StartInfo.RedirectStandardError = true;
 
-			using (Process DocToolProcess = new Process())
-			{
-				DocToolProcess.StartInfo.WorkingDirectory = EngineDir;
-				DocToolProcess.StartInfo.FileName = DocToolPath;
-				DocToolProcess.StartInfo.Arguments = (bBlueprint ? "BlueprintAPI" : "API") + "\\* -lang=INT -t=DefaultAPI.html -v=warn";
-				DocToolProcess.StartInfo.UseShellExecute = false;
-				DocToolProcess.StartInfo.RedirectStandardOutput = true;
-				DocToolProcess.StartInfo.RedirectStandardError = true;
+					DocToolProcess.OutputDataReceived += new DataReceivedEventHandler(ProcessOutputReceived);
+					DocToolProcess.ErrorDataReceived += new DataReceivedEventHandler(ProcessOutputReceived);
 
-				DocToolProcess.OutputDataReceived += new DataReceivedEventHandler(ProcessOutputReceived);
-				DocToolProcess.ErrorDataReceived += new DataReceivedEventHandler(ProcessOutputReceived);
+					try
+					{
+						DocToolProcess.Start();
+						DocToolProcess.BeginOutputReadLine();
+						DocToolProcess.BeginErrorReadLine();
+						DocToolProcess.WaitForExit();
+					}
+					catch (Exception Ex)
+					{
+						Console.WriteLine(Ex.ToString() + "\n" + Ex.StackTrace);
+						return false;
+					}
 
-				try
-				{
-					DocToolProcess.Start();
-					DocToolProcess.BeginOutputReadLine();
-					DocToolProcess.BeginErrorReadLine();
-					DocToolProcess.WaitForExit();
-					return DocToolProcess.ExitCode == 0;
-				}
-				catch (Exception Ex)
-				{
-					Console.WriteLine(Ex.ToString() + "\n" + Ex.StackTrace);
-					return false;
+					if (DocToolProcess.ExitCode != 0)
+					{
+						return false;
+					}
 				}
 			}
-		}
-
-		public static void CleanChm(string ChmPath)
-		{
-			Console.WriteLine("Cleaning '{0}'", ChmPath);
-			if (Directory.Exists(ChmPath))
+			if ((Actions & BuildActions.Archive) != 0)
 			{
-				foreach (string SubDir in Directory.EnumerateDirectories(ChmPath))
-				{
-					Directory.Delete(SubDir, true);
-				}
-			}
-		}
+				Console.WriteLine("Creating archive '{0}'", ArchivePath);
 
-		public static bool BuildChm(string ChmCompilerPath, string BaseHtmlDir, string ChmDir, bool bBlueprint)
-		{
-			string ProjectFileName = (bBlueprint ? "BlueprintAPI.hhp" : "API.hhp");
-			Console.WriteLine("Searching for CHM input files...");
+				List<string> InputFiles = new List<string>();
+				Utility.FindFilesForTar(HtmlDir, "Include", InputFiles, true);
+				Utility.FindFilesForTar(HtmlDir, "INT\\" + HtmlSuffix, InputFiles, true);
+				Utility.FindFilesForTar(HtmlDir, "Images", InputFiles, false);
 
-			// Build a list of all the files we want to copy
-			List<string> FilePaths = new List<string>();
-			List<string> DirectoryPaths = new List<string>();
-			if (bBlueprint)
-			{
-				Utility.FindRelativeContents(BaseHtmlDir, "Images\\bp_api*", false, FilePaths, DirectoryPaths);
+				Utility.CreateTgzFromFiles(ArchivePath, HtmlDir, InputFiles);
 			}
-			else
-			{
-				Utility.FindRelativeContents(BaseHtmlDir, "Images\\api*", false, FilePaths, DirectoryPaths);
-			}
-			Utility.FindRelativeContents(BaseHtmlDir, "Include\\*", true, FilePaths, DirectoryPaths);
+			return true;
+		}
 
-			// Find all the HTML files
-			List<string> HtmlFilePaths = new List<string>();
-			if (bBlueprint)
+		static bool BuildChm(string ChmCompilerPath, string ChmFileName, string ContentsFileName, string IndexFileName, string BaseHtmlDir, bool bBlueprint, BuildActions Actions)
+		{
+			if((Actions & BuildActions.Clean) != 0)
 			{
-				Utility.FindRelativeContents(BaseHtmlDir, "INT\\BlueprintAPI\\*.html", true, HtmlFilePaths, DirectoryPaths);
+				Console.WriteLine("Cleaning '{0}'", ChmFileName);
+				Utility.SafeDeleteFile(ChmFileName);
+				Utility.SafeDeleteDirectory(Path.Combine(Path.GetDirectoryName(ChmFileName), Path.GetFileNameWithoutExtension(ChmFileName)));
 			}
-			else 
+			if ((Actions & BuildActions.Build) != 0)
 			{
-				Utility.FindRelativeContents(BaseHtmlDir, "INT\\API\\*.html", true, HtmlFilePaths, DirectoryPaths);
-			}
+				string ProjectFileName = (bBlueprint ? "BlueprintAPI.hhp" : "API.hhp");
+				Console.WriteLine("Searching for CHM input files...");
 
-			// Create all the target directories
-			foreach (string DirectoryPath in DirectoryPaths)
-			{
-				Utility.SafeCreateDirectory(Path.Combine(ChmDir, DirectoryPath));
-			}
+				// Build a list of all the files we want to copy
+				List<string> FilePaths = new List<string>();
+				List<string> DirectoryPaths = new List<string>();
+				Utility.FindRelativeContents(BaseHtmlDir, "Include\\*", true, FilePaths, DirectoryPaths);
 
-			// Copy all the files across
-			using (Tracker CopyTracker = new Tracker("Copying support files...", FilePaths.Count))
-			{
-				foreach(int Idx in CopyTracker.Indices)
+				// Find all the HTML files
+				if (bBlueprint)
+				{
+					Utility.FindRelativeContents(BaseHtmlDir, "Images\\bp_api*", false, FilePaths, DirectoryPaths);
+					Utility.FindRelativeContents(BaseHtmlDir, "INT\\BlueprintAPI\\*.html", true, FilePaths, DirectoryPaths);
+					CreateChm(ChmCompilerPath, ChmFileName, "UE4 Blueprint API Documentation", "INT\\BlueprintAPI\\index.html", ContentsFileName, IndexFileName, BaseHtmlDir, FilePaths);
+				}
+				else
 				{
-					Utility.SafeCopyFile(Path.Combine(BaseHtmlDir, FilePaths[Idx]), Path.Combine(ChmDir, FilePaths[Idx]));
+					Utility.FindRelativeContents(BaseHtmlDir, "Images\\api*", false, FilePaths, DirectoryPaths);
+					Utility.FindRelativeContents(BaseHtmlDir, "INT\\API\\*.html", true, FilePaths, DirectoryPaths);
+					CreateChm(ChmCompilerPath, ChmFileName, "UE4 API Documentation", "INT\\API\\index.html", ContentsFileName, IndexFileName, BaseHtmlDir, FilePaths);
 				}
 			}
+			return true;
+		}
+
+		static private void CreateChm(string ChmCompilerPath, string ChmFileName, string Title, string DefaultTopicPath, string ContentsFileName, string IndexFileName, string SourceDir, List<string> FileNames)
+		{
+			string ProjectName = Path.GetFileNameWithoutExtension(ChmFileName);
+
+			// Create an intermediate directory
+			string IntermediateDir = Path.Combine(Path.GetDirectoryName(ChmFileName), ProjectName);
+			Directory.CreateDirectory(IntermediateDir);
+			Utility.SafeDeleteDirectoryContents(IntermediateDir, true);
 
-			// Copy the HTML files across, fixing up the HTML for display in the CHM window
-			using (Tracker HtmlTracker = new Tracker("Copying html files...", HtmlFilePaths.Count))
+			// Copy all the files across
+			using (Tracker CopyTracker = new Tracker("Copying files...", FileNames.Count))
 			{
-				foreach(int Idx in HtmlTracker.Indices)
+				foreach (int Idx in CopyTracker.Indices)
 				{
-					string HtmlFilePath = HtmlFilePaths[Idx];
-					string HtmlText = File.ReadAllText(Path.Combine(BaseHtmlDir, HtmlFilePath));
+					string SourceFileName = Path.Combine(SourceDir, FileNames[Idx]);
+					string TargetFileName = Path.Combine(IntermediateDir, FileNames[Idx]);
+
+					Directory.CreateDirectory(Path.GetDirectoryName(TargetFileName));
+
+					if (SourceFileName.EndsWith(".html", StringComparison.InvariantCultureIgnoreCase))
+					{
+						string HtmlText = File.ReadAllText(SourceFileName);
 
-					HtmlText = HtmlText.Replace("<?xml version=\"1.0\" encoding=\"UTF-8\"?>\r\n", "");
-					HtmlText = HtmlText.Replace("<div id=\"crumbs_bg\"></div>", "");
+						HtmlText = HtmlText.Replace("<?xml version=\"1.0\" encoding=\"UTF-8\"?>\r\n", "");
+						HtmlText = HtmlText.Replace("<div id=\"crumbs_bg\"></div>", "");
 
-					const string HeaderEndText = "<!-- end head -->";
-					int HeaderMinIdx = HtmlText.IndexOf("<div id=\"head\">");
-					int HeaderMaxIdx = HtmlText.IndexOf(HeaderEndText);
-					HtmlText = HtmlText.Remove(HeaderMinIdx, HeaderMaxIdx + HeaderEndText.Length - HeaderMinIdx);
+						const string HeaderEndText = "<!-- end head -->";
+						int HeaderMinIdx = HtmlText.IndexOf("<div id=\"head\">");
+						int HeaderMaxIdx = HtmlText.IndexOf(HeaderEndText);
+						HtmlText = HtmlText.Remove(HeaderMinIdx, HeaderMaxIdx + HeaderEndText.Length - HeaderMinIdx);
 
-					int CrumbsMinIdx = HtmlText.IndexOf("<div class=\"crumbs\">");
-					int HomeMinIdx = HtmlText.IndexOf("<strong>", CrumbsMinIdx);
-					int HomeMaxIdx = HtmlText.IndexOf("&gt;", HomeMinIdx) + 4;
-					HtmlText = HtmlText.Remove(HomeMinIdx, HomeMaxIdx - HomeMinIdx);
+						int CrumbsMinIdx = HtmlText.IndexOf("<div class=\"crumbs\">");
+						int HomeMinIdx = HtmlText.IndexOf("<strong>", CrumbsMinIdx);
+						int HomeMaxIdx = HtmlText.IndexOf("&gt;", HomeMinIdx) + 4;
+						HtmlText = HtmlText.Remove(HomeMinIdx, HomeMaxIdx - HomeMinIdx);
 
-					File.WriteAllText(Path.Combine(ChmDir, HtmlFilePath), HtmlText);
+						File.WriteAllText(TargetFileName, HtmlText);
+					}
+					else
+					{
+						Utility.SafeCopyFile(SourceFileName, TargetFileName);
+					}
 				}
 			}
 
+			// Copy the contents and index files
+			Utility.SafeCopyFile(ContentsFileName, Path.Combine(IntermediateDir, ProjectName + ".hhc"));
+			Utility.SafeCopyFile(IndexFileName, Path.Combine(IntermediateDir, ProjectName + ".hhk"));
+
 			// Write the project file
-			using (StreamWriter Writer = new StreamWriter(Path.Combine(ChmDir, ProjectFileName)))
+			string ProjectFileName = Path.Combine(IntermediateDir, ProjectName + ".hhp");
+			using (StreamWriter Writer = new StreamWriter(ProjectFileName))
 			{
 				Writer.WriteLine("[OPTIONS]");
-				if (bBlueprint)
-				{
-					Writer.WriteLine("Title=UE4 Blueprint API Documentation");
-					Writer.WriteLine("Binary TOC=Yes");
-					Writer.WriteLine("Compatibility=1.1 or later");
-					Writer.WriteLine("Compiled file=BlueprintAPI.chm");
-					Writer.WriteLine("Contents file=BlueprintAPI.hhc");
-					Writer.WriteLine("Index file=BlueprintAPI.hhk");
-					Writer.WriteLine("Default topic=INT\\BlueprintAPI\\index.html");
-					Writer.WriteLine("Full-text search=Yes");
-					Writer.WriteLine("Display compile progress=Yes");
-					Writer.WriteLine("Language=0x409 English (United States)");
-				}
-				else
-				{
-					Writer.WriteLine("Title=UE4 API Documentation");
-					Writer.WriteLine("Binary TOC=Yes");
-					Writer.WriteLine("Compatibility=1.1 or later");
-					Writer.WriteLine("Compiled file=API.chm");
-					Writer.WriteLine("Contents file=API.hhc");
-					Writer.WriteLine("Index file=API.hhk");
-					Writer.WriteLine("Default topic=INT\\API\\index.html");
-					Writer.WriteLine("Full-text search=Yes");
-					Writer.WriteLine("Display compile progress=Yes");
-					Writer.WriteLine("Language=0x409 English (United States)");
-				}
+				Writer.WriteLine("Title={0}", Title);
+				Writer.WriteLine("Binary TOC=Yes");
+				Writer.WriteLine("Compatibility=1.1 or later");
+				Writer.WriteLine("Compiled file={0}", Path.GetFileName(ChmFileName));
+				Writer.WriteLine("Contents file={0}.hhc", ProjectName);
+				Writer.WriteLine("Index file={0}.hhk", ProjectName);
+				Writer.WriteLine("Default topic={0}", DefaultTopicPath);
+				Writer.WriteLine("Full-text search=Yes");
+				Writer.WriteLine("Display compile progress=Yes");
+				Writer.WriteLine("Language=0x409 English (United States)");
 				Writer.WriteLine();
 				Writer.WriteLine("[FILES]");
-				foreach (string FilePath in FilePaths)
-				{
-					Writer.WriteLine(FilePath);
-				}
-				foreach (string HtmlFilePath in HtmlFilePaths)
+				foreach (string FileName in FileNames)
 				{
-					Writer.WriteLine(HtmlFilePath);
+					Writer.WriteLine(FileName);
 				}
 			}
 
@@ -1276,9 +1203,9 @@ namespace APIDocTool
 			Console.WriteLine("Compiling CHM file...");
 			using (Process CompilerProcess = new Process())
 			{
-				CompilerProcess.StartInfo.WorkingDirectory = ChmDir;
+				CompilerProcess.StartInfo.WorkingDirectory = IntermediateDir;
 				CompilerProcess.StartInfo.FileName = ChmCompilerPath;
-				CompilerProcess.StartInfo.Arguments = ProjectFileName;
+				CompilerProcess.StartInfo.Arguments = Path.GetFileName(ProjectFileName);
 				CompilerProcess.StartInfo.UseShellExecute = false;
 				CompilerProcess.StartInfo.RedirectStandardOutput = true;
 				CompilerProcess.StartInfo.RedirectStandardError = true;
@@ -1289,7 +1216,9 @@ namespace APIDocTool
 				CompilerProcess.BeginErrorReadLine();
 				CompilerProcess.WaitForExit();
 			}
-			return true;
+
+			// Copy it to the final output
+			Utility.SafeCopyFile(Path.Combine(IntermediateDir, ProjectName + ".chm"), ChmFileName);
 		}
 
 		static private void ProcessOutputReceived(Object Sender, DataReceivedEventArgs Line)
diff --git a/Engine/Source/Programs/UnrealDocTool/APIDocTool/APIDocTool/Utility.cs b/Engine/Source/Programs/UnrealDocTool/APIDocTool/APIDocTool/Utility.cs
--- a/Engine/Source/Programs/UnrealDocTool/APIDocTool/APIDocTool/Utility.cs
+++ b/Engine/Source/Programs/UnrealDocTool/APIDocTool/APIDocTool/Utility.cs
@@ -9,6 +9,7 @@ using System.Threading.Tasks;
 using System.Web;
 using System.Xml;
 using DoxygenLib;
+using System.IO.Compression;
 
 namespace APIDocTool
 {
@@ -425,5 +426,142 @@ namespace APIDocTool
 			if(Result == '/') Result = '\\';
 			return Result;
 		}
+
+		public static void CreateTgzFromDir(string ArchivePath, string BaseDir)
+		{
+			List<string> InputFiles = new List<string>();
+			Utility.FindFilesForTar(BaseDir, "", InputFiles, true);
+			Utility.CreateTgzFromFiles(ArchivePath, BaseDir, InputFiles);
+		}
+
+		public static void CreateTgzFromFiles(string TarPath, string BaseDir, IEnumerable<string> InputFiles)
+		{
+			using (FileStream TarFileStream = new FileStream(TarPath, FileMode.Create))
+			{
+				using (GZipStream ZipStream = new GZipStream(TarFileStream, CompressionMode.Compress))
+				{
+					using (BinaryWriter TarWriter = new BinaryWriter(ZipStream))
+					{
+						foreach (string InputFile in InputFiles)
+						{
+							byte[] FileData = File.ReadAllBytes(Path.Combine(BaseDir, InputFile));
+							WriteFileToTar(TarWriter, InputFile, FileData);
+						}
+						TarWriter.Write(new byte[1024]);
+					}
+				}
+			}
+		}
+
+		public static void FindFilesForTar(string BaseDir, string BaseTarPath, List<string> Files, bool bRecursive)
+		{
+			DirectoryInfo Dir = new DirectoryInfo(Path.Combine(BaseDir, BaseTarPath));
+			if (Dir.Exists)
+			{
+				FindFilesForTar(Dir, BaseTarPath, Files, bRecursive);
+			}
+		}
+
+		public static void FindFilesForTar(DirectoryInfo Dir, string BaseTarPath, List<string> Files, bool bRecursive)
+		{
+			if (bRecursive)
+			{
+				foreach (DirectoryInfo ChildDir in Dir.GetDirectories())
+				{
+					FindFilesForTar(ChildDir, Path.Combine(BaseTarPath, ChildDir.Name), Files, bRecursive);
+				}
+			}
+			foreach (FileInfo FileInfo in Dir.GetFiles())
+			{
+				Files.Add(Path.Combine(BaseTarPath, FileInfo.Name));
+			}
+		}
+
+		const int TarFileNameFieldLength = 99;
+
+		public static void WriteFileToTar(BinaryWriter Writer, string FilePath, byte[] FileData)
+		{
+			string NormalizedFilePath = FilePath.Replace('\\', '/');
+
+			// Generate a long header record if necessary
+			if (NormalizedFilePath.Length > TarFileNameFieldLength)
+			{
+				WriteTarHeader(Writer, "././@LongLink", NormalizedFilePath.Length + 1, 'L');
+
+				byte[] FilePathBytes = new byte[NormalizedFilePath.Length + 1];
+				EncodeTarString(FilePathBytes, 0, FilePathBytes.Length, NormalizedFilePath);
+
+				WriteTarBlocks(Writer, FilePathBytes);
+			}
+
+			// Now write the normal file
+			WriteTarHeader(Writer, NormalizedFilePath, FileData.Length, '0');
+			WriteTarBlocks(Writer, FileData);
+		}
+
+		static void WriteTarHeader(BinaryWriter Writer, string FileName, int Length, char Type)
+		{
+			// Write the name
+			byte[] Header = new byte[512];
+
+			// Name
+			EncodeTarString(Header, 0, TarFileNameFieldLength, FileName);
+
+			// File mode
+			EncodeTarValue(Header, 100, 7, 0x81ff);
+
+			// UID
+			EncodeTarValue(Header, 108, 7, 0);
+
+			// GID
+			EncodeTarValue(Header, 116, 7, 0);
+
+			// File size
+			EncodeTarValue(Header, 124, 12, Length);
+
+			// File timestamp
+			EncodeTarValue(Header, 136, 12, 0);
+
+			// File type ('0' = Regular file)
+			Header[156] = (byte)Type;
+
+			// Generate the the checksum (seed the checksum as if the field was initialized with 8 spaces)
+			int Checksum = (byte)' ' * 8;
+			for (int Idx = 0; Idx < Header.Length; Idx++) Checksum += Header[Idx];
+			EncodeTarValue(Header, 148, 7, Checksum);
+
+			// Write the header
+			Writer.Write(Header);
+		}
+
+		static void WriteTarBlocks(BinaryWriter Writer, byte[] Data)
+		{
+			Writer.Write(Data);
+			Writer.Write(new byte[(512 - (Data.Length % 512)) % 512]);
+		}
+
+		static void EncodeTarString(byte[] Data, int Offset, int Length, string Text)
+		{
+			// Convert the Text to ASCII
+			for (int Idx = 0; Idx < Math.Min(Length, Text.Length); Idx++)
+			{
+				if (Text[Idx] == 0 || Text[Idx] > 0x7f)
+				{
+					throw new Exception("Invalid character in TAR string");
+				}
+				else
+				{
+					Data[Offset + Idx] = (byte)Text[Idx];
+				}
+			}
+		}
+
+		static void EncodeTarValue(byte[] Data, int Offset, int Length, int Value)
+		{
+			// Write the value as an octal string followed by a space
+			string ValueString = Convert.ToString(Value, 8) + " ";
+			while (ValueString.Length < Length) ValueString = " " + ValueString;
+			EncodeTarString(Data, Offset, Length, ValueString);
+		}
 	}
 }
