Allow stripping URLs which have anchors in them

diff --git a/Engine/Source/Programs/UnrealDocTool/DoxygenLib/Markdown.cs b/Engine/Source/Programs/UnrealDocTool/DoxygenLib/Markdown.cs
--- a/Engine/Source/Programs/UnrealDocTool/DoxygenLib/Markdown.cs
+++ b/Engine/Source/Programs/UnrealDocTool/DoxygenLib/Markdown.cs
@@ -137,12 +137,17 @@ namespace DoxygenLib
 					{
 						LinkText = "";
 						LinkUrl = LinkUrl.Substring(WebPrefixIdx + WebPrefix.Length);
-						if (LinkUrl.EndsWith(".html"))
+
+						int AnchorIdx = LinkUrl.LastIndexOf('#');
+						if(AnchorIdx == -1)
+						{
+							LinkUrl = RemoveDefaultPageName(LinkUrl);
+						}
+						else
 						{
-							LinkUrl = LinkUrl.Substring(0, LinkUrl.LastIndexOf('/') + 1);
+							LinkUrl = RemoveDefaultPageName(LinkUrl.Substring(0, AnchorIdx)) + LinkUrl.Substring(AnchorIdx);
 						}
 					}
-
 					Writer.WriteLink(LinkText, LinkUrl);
 					break;
 				case "ref":
@@ -226,6 +231,22 @@ namespace DoxygenLib
 			}
 		}
 
+		/** 
+		 * Removes the default page name from a URL (index.html)
+		 * 
+		 * @param LinkUrl			URL for the page
+		 * @return URL with the page name removed
+		 */
+		static string RemoveDefaultPageName(string LinkUrl)
+		{
+			string DefaultPageName = "/index.html";
+			if (LinkUrl.EndsWith(DefaultPageName))
+			{
+				return LinkUrl.Substring(0, LinkUrl.Length - DefaultPageName.Length - 1);
+			}
+			return LinkUrl;
+		}
+
 		/** 
 		 * Converts all children of the given XML node
 		 * 
