Added support for using multiple params with the same name to specify arrays in object templates

diff --git a/Engine/Source/Programs/UnrealDocTool/MarkdownMode/source.extension.vsixmanifest b/Engine/Source/Programs/UnrealDocTool/MarkdownMode/source.extension.vsixmanifest
--- a/Engine/Source/Programs/UnrealDocTool/MarkdownMode/source.extension.vsixmanifest
+++ b/Engine/Source/Programs/UnrealDocTool/MarkdownMode/source.extension.vsixmanifest
@@ -1,7 +1,7 @@
 <?xml version="1.0" encoding="utf-8"?>
 <PackageManifest Version="2.0.0" xmlns="http://schemas.microsoft.com/developer/vsx-schema/2011" xmlns:d="http://schemas.microsoft.com/developer/vsx-schema-design/2011">
   <Metadata>
-    <Identity Id="Epic MarkdownMode" Version="0.221" Language="en-US" Publisher="Epic Games" />
+    <Identity Id="Epic MarkdownMode" Version="0.222" Language="en-US" Publisher="Epic Games" />
     <DisplayName>Epic MarkdownMode</DisplayName>
     <Description xml:space="preserve">A language service for editing Markdown files.</Description>
     <Icon>markdown-icon.png</Icon>
diff --git a/Engine/Source/Programs/UnrealDocTool/UnrealDocTool/MarkdownSharp/EpicMarkdown/EMObject.cs b/Engine/Source/Programs/UnrealDocTool/UnrealDocTool/MarkdownSharp/EpicMarkdown/EMObject.cs
--- a/Engine/Source/Programs/UnrealDocTool/UnrealDocTool/MarkdownSharp/EpicMarkdown/EMObject.cs
+++ b/Engine/Source/Programs/UnrealDocTool/UnrealDocTool/MarkdownSharp/EpicMarkdown/EMObject.cs
@@ -16,29 +16,53 @@ namespace MarkdownSharp.EpicMarkdown
     public class EMObject : EMElement
     {
         private readonly TemplateFile template;
-        private readonly Dictionary<string, string> literalParams;
-        private readonly Dictionary<string, EMObjectParam> markdownParams;
+        private readonly Dictionary<string, List<string>> literalParams;
+        private readonly Dictionary<string, List<EMObjectParam>> markdownParams;
 
         private EMObject(EMDocument doc, EMElementOrigin origin, EMElement parent, string templateFile)
             : base(doc, origin, parent)
         {
             template = Templates.GetCached(templateFile);
-            this.literalParams = new Dictionary<string, string>();
-            this.markdownParams = new Dictionary<string, EMObjectParam>();
+            this.literalParams = new Dictionary<string, List<string>>();
+            this.markdownParams = new Dictionary<string, List<EMObjectParam>>();
         }
 
         public override void AppendHTML(StringBuilder builder, Stack<EMInclude> includesStack, TransformationData data)
         {
             var parameters = new Hash();
 
-            foreach (var literal in literalParams)
+            foreach (var paramPair in literalParams)
             {
-                parameters.Add(literal.Key, literal.Value);
+                List<string> paramStrings = new List<string>();
+                if (paramPair.Value.Count > 1)
+                {
+                    foreach (var param in paramPair.Value)
+                    {
+                        paramStrings.Add(param);
+                    }
+                    parameters.Add(paramPair.Key, paramStrings.ToArray());
+                }
+                else if (paramPair.Value.Count > 0)
+                {
+                    parameters.Add(paramPair.Key, paramPair.Value[0]);
+                }
             }
 
-            foreach (var markdownParam in markdownParams)
+            foreach (var paramPair in markdownParams)
             {
-                parameters.Add(markdownParam.Key, markdownParam.Value.Elements.GetInnerHTML(includesStack, data));
+                List<string> paramStrings = new List<string>();
+                if (paramPair.Value.Count > 1)
+                {
+                    foreach (var param in paramPair.Value)
+                    {
+                        paramStrings.Add(param.Elements.GetInnerHTML(includesStack, data));
+                    }
+                    parameters.Add(paramPair.Key, paramStrings.ToArray());
+                }
+                else if (paramPair.Value.Count > 0)
+                {
+                    parameters.Add(paramPair.Key, paramPair.Value[0].Elements.GetInnerHTML(includesStack, data));
+                }
             }
 
             builder.Append(template.Render(parameters));
@@ -78,7 +102,16 @@ namespace MarkdownSharp.EpicMarkdown
                     text = text.Trim('\n');
 
                     text = Markdown.OutdentIfPossible(text);
-                    literalParams.Add(paramName, text);
+                    if(literalParams.ContainsKey(paramName))
+                    {
+                        literalParams[paramName].Add(text);
+                    }
+                    else
+                    {
+                        List<string> literalStrings = new List<string>();
+                        literalStrings.Add(text);
+                        literalParams.Add(paramName, literalStrings);
+                    }
                 }
                 else
                 {
@@ -89,7 +122,16 @@ namespace MarkdownSharp.EpicMarkdown
                     var param = EMObjectParam.CreateParam(origin, Document, this, data,
                         tagMatch, paramName, paramMatch.Groups["indentation"].Value);
 
-                    markdownParams.Add(paramName, param);
+                    if (literalParams.ContainsKey(paramName))
+                    {
+                        markdownParams[paramName].Add(param);
+                    }
+                    else
+                    {
+                        List<EMObjectParam> paramList = new List<EMObjectParam>();
+                        paramList.Add(param);
+                        markdownParams.Add(paramName, paramList);
+                    }
                 }
             }
         }
@@ -98,9 +140,10 @@ namespace MarkdownSharp.EpicMarkdown
         {
             context.PreChildrenVisit(this);
 
-            foreach (var param in markdownParams)
+            foreach (var paramPair in markdownParams)
             {
-                param.Value.ForEachWithContext(context);
+                foreach(var param in paramPair.Value)
+                param.ForEachWithContext(context);
             }
 
             context.PostChildrenVisit(this);
