Escape ! before markdown links.

diff --git a/Engine/Source/Programs/UnrealDocTool/DoxygenLib/Markdown.cs b/Engine/Source/Programs/UnrealDocTool/DoxygenLib/Markdown.cs
--- a/Engine/Source/Programs/UnrealDocTool/DoxygenLib/Markdown.cs
+++ b/Engine/Source/Programs/UnrealDocTool/DoxygenLib/Markdown.cs
@@ -71,6 +71,17 @@ namespace DoxygenLib
 
 		public void WriteLink(string LinkText, string LinkUrl)
 		{
+			// Make sure the last character wasn't a !, otherwise it'll be treated as an attachment reference
+			if(Buffer.Length > 0 && Buffer[Buffer.Length - 1] == '!')
+			{
+				if(Buffer.Length < 2 || Buffer[Buffer.Length - 2] != '\\')
+				{
+					Buffer.Remove(Buffer.Length - 1, 1);
+					Buffer.Append("\\!");
+				}
+			}
+
+			// Write the link text
 			WriteFormat("[{0}]({1})", LinkText, LinkUrl);
 		}
 
