Fix URLs being truncated when translated from full docs.unrealengine.com links.

diff --git a/Engine/Source/Programs/UnrealDocTool/DoxygenLib/Markdown.cs b/Engine/Source/Programs/UnrealDocTool/DoxygenLib/Markdown.cs
--- a/Engine/Source/Programs/UnrealDocTool/DoxygenLib/Markdown.cs
+++ b/Engine/Source/Programs/UnrealDocTool/DoxygenLib/Markdown.cs
@@ -253,7 +253,7 @@ namespace DoxygenLib
 			string DefaultPageName = "/index.html";
 			if (LinkUrl.EndsWith(DefaultPageName))
 			{
-				return LinkUrl.Substring(0, LinkUrl.Length - DefaultPageName.Length - 1);
+				return LinkUrl.Substring(0, LinkUrl.Length - DefaultPageName.Length + 1);
 			}
 			return LinkUrl;
 		}
