Bypass common constructs for API pages, which UnrealDocTool treats as bad markup. (eg. "0% to 100%" is not a variable reference, [something] (something else) should not have to be an URL).

diff --git a/Engine/Source/Programs/UnrealDocTool/UnrealDocTool/MarkdownSharp/EpicMarkdown/EMDocument.cs b/Engine/Source/Programs/UnrealDocTool/UnrealDocTool/MarkdownSharp/EpicMarkdown/EMDocument.cs
--- a/Engine/Source/Programs/UnrealDocTool/UnrealDocTool/MarkdownSharp/EpicMarkdown/EMDocument.cs
+++ b/Engine/Source/Programs/UnrealDocTool/UnrealDocTool/MarkdownSharp/EpicMarkdown/EMDocument.cs
@@ -126,6 +126,11 @@ namespace MarkdownSharp.EpicMarkdown
             return LocalPath;
         }
 
+		public bool PerformStrictConversion()
+		{
+			return !GetAbsoluteMarkdownPath().Replace('\\', '/').StartsWith("API/");
+		}
+
         public string GetAbsolutePath(string relativePath)
         {
             return
diff --git a/Engine/Source/Programs/UnrealDocTool/UnrealDocTool/MarkdownSharp/EpicMarkdown/EMLink.cs b/Engine/Source/Programs/UnrealDocTool/UnrealDocTool/MarkdownSharp/EpicMarkdown/EMLink.cs
--- a/Engine/Source/Programs/UnrealDocTool/UnrealDocTool/MarkdownSharp/EpicMarkdown/EMLink.cs
+++ b/Engine/Source/Programs/UnrealDocTool/UnrealDocTool/MarkdownSharp/EpicMarkdown/EMLink.cs
@@ -227,7 +227,16 @@ namespace MarkdownSharp.EpicMarkdown
             }
             catch (EMPathVerificationException e)
             {
-                return new EMErrorElement(doc, orig, parent, e.AddToErrorList(data, orig.Text));
+				if (doc.PerformStrictConversion())
+				{
+					return new EMErrorElement(doc, orig, parent, e.AddToErrorList(data, orig.Text));
+				}
+				else
+				{
+					EMFormattedText Text = new EMFormattedText(doc, orig, parent);
+					Text.Parse(orig.Text, data);
+					return Text;
+				}
             }
         }
 
diff --git a/Engine/Source/Programs/UnrealDocTool/UnrealDocTool/MarkdownSharp/Preprocessor/VariableManager.cs b/Engine/Source/Programs/UnrealDocTool/UnrealDocTool/MarkdownSharp/Preprocessor/VariableManager.cs
--- a/Engine/Source/Programs/UnrealDocTool/UnrealDocTool/MarkdownSharp/Preprocessor/VariableManager.cs
+++ b/Engine/Source/Programs/UnrealDocTool/UnrealDocTool/MarkdownSharp/Preprocessor/VariableManager.cs
@@ -225,18 +225,20 @@ namespace MarkdownSharp.Preprocessor
 
                 if (!data.ProcessedDocumentCache.TryGetLinkedFileVariable(currentFile, variableName, out outString))
                 {
-                    //error
-                    errorId = data.ErrorList.Count;
-                    data.ErrorList.Add(
-                        Markdown.GenerateError(
-                            Language.Message("VariableOrMetadataNotFoundInFile", variableName, currentFile.FullName),
-                            MessageClass.Error,
-                            OnErrorStringToUseForMatch,
-                            errorId,
-                            data));
-
-                    isProblem = true;
-
+					if (doc.PerformStrictConversion())
+					{
+						//error
+						errorId = data.ErrorList.Count;
+						data.ErrorList.Add(
+							Markdown.GenerateError(
+								Language.Message("VariableOrMetadataNotFoundInFile", variableName, currentFile.FullName),
+								MessageClass.Error,
+								OnErrorStringToUseForMatch,
+								errorId,
+								data));
+
+						isProblem = true;
+					}
                     outString = OnErrorStringToUseForMatch.Replace("%", "&#37;");
                 }
                 else
