Another attempt to fix UDN conversion errors for text containing ! characters before symbol references.

diff --git a/Engine/Source/Programs/UnrealDocTool/DoxygenLib/Markdown.cs b/Engine/Source/Programs/UnrealDocTool/DoxygenLib/Markdown.cs
--- a/Engine/Source/Programs/UnrealDocTool/DoxygenLib/Markdown.cs
+++ b/Engine/Source/Programs/UnrealDocTool/DoxygenLib/Markdown.cs
@@ -71,17 +71,6 @@ namespace DoxygenLib
 
 		public void WriteLink(string LinkText, string LinkUrl)
 		{
-			// Make sure the last character wasn't a !, otherwise it'll be treated as an attachment reference
-			if(Buffer.Length > 0 && Buffer[Buffer.Length - 1] == '!')
-			{
-				if(Buffer.Length < 2 || Buffer[Buffer.Length - 2] != '\\')
-				{
-					Buffer.Remove(Buffer.Length - 1, 1);
-					Buffer.Append("\\!");
-				}
-			}
-
-			// Write the link text
 			WriteFormat("[{0}]({1})", LinkText, LinkUrl);
 		}
 
@@ -276,7 +265,7 @@ namespace DoxygenLib
 			}
 			else
 			{
-				Writer.Write(Node.InnerText);
+				Writer.Write(EscapeXmlText(Node.InnerText));
 			}
 		}
 
@@ -391,6 +380,30 @@ namespace DoxygenLib
 			return Result.ToString();
 		}
 
+		/** 
+		 * Escape xml text to markdown text.
+		 * 
+		 * @param Text			String to escape
+		 * @return				Escaped string
+		 */
+		public static string EscapeXmlText(string Text)
+		{
+			StringBuilder Result = new StringBuilder();
+			for (int Idx = 0; Idx < Text.Length; Idx++)
+			{
+				if (Text[Idx] == '!')
+				{
+					Result.Append("&#33;");
+				}
+				else
+				{
+					Result.Append(Text[Idx]);
+				}
+			}
+			return Result.ToString();
+		}
+
+
 		/** 
 		 * Truncates a line to a given width, respecting embedded links.
          * 
