Experimental - remove argument lists for functions (display argument lists whenever a function is overloaded), and create a separate group for protected members.

diff --git a/Engine/Source/Programs/UnrealDocTool/APIDocTool/APIDocTool/APIFunction.cs b/Engine/Source/Programs/UnrealDocTool/APIDocTool/APIDocTool/APIFunction.cs
--- a/Engine/Source/Programs/UnrealDocTool/APIDocTool/APIDocTool/APIFunction.cs
+++ b/Engine/Source/Programs/UnrealDocTool/APIDocTool/APIDocTool/APIFunction.cs
@@ -770,13 +770,29 @@ namespace APIDocTool
 			}
 			Writer.WriteParam("icons", ItemIcons);
 
+			// If the function is overloaded, format the display name with a list of arguments
+			string DisplayName = Markdown.EscapeText(Name);
+			if (Parent is APIFunctionGroup)
+			{
+				string ArgumentList = String.Format("&#40;{0}&#41;", String.Join(", ", Parameters.Select(x => x.TypeText)));
+
+				int MaxArgumentListLength = Math.Max(50 - DisplayName.Length, 10);
+				if (ArgumentList.Length > MaxArgumentListLength + 4)
+				{
+					ArgumentList = Markdown.Truncate(ArgumentList, MaxArgumentListLength, "...");
+				}
+
+				DisplayName += ArgumentList;
+			}
+			DisplayName = DisplayName.Replace(" ", "&nbsp;");
+
 			// Write the name
-			Writer.WriteParam("name", Name);
+			Writer.WriteParam("name", DisplayName);
 			Writer.WriteParam("link", "[RELATIVE:" + LinkPath + "]");
 
 			// Add the parameter section if need be
 			Writer.EnterParam("arguments");
-			if (Parameters.Count > 0)
+			if (Parameters.Count > 0 && false)
 			{
 				Writer.WriteEscapedLine("(  ");
 				for (int Idx = 0; Idx < Parameters.Count; Idx++)
diff --git a/Engine/Source/Programs/UnrealDocTool/APIDocTool/APIDocTool/APIRecord.cs b/Engine/Source/Programs/UnrealDocTool/APIDocTool/APIDocTool/APIRecord.cs
--- a/Engine/Source/Programs/UnrealDocTool/APIDocTool/APIDocTool/APIRecord.cs
+++ b/Engine/Source/Programs/UnrealDocTool/APIDocTool/APIDocTool/APIRecord.cs
@@ -378,7 +378,7 @@ namespace APIDocTool
 				List<APIFunction> AllFunctions = new List<APIFunction>();
 				AllFunctions.AddRange(Children.OfType<APIFunction>().Where(x => x.Protection != APIProtection.Private && !x.IsDeprecated()));
 				AllFunctions.AddRange(Children.OfType<APIFunctionGroup>().SelectMany(x => x.Children.OfType<APIFunction>()).Where(x => x.Protection != APIProtection.Private && !x.IsDeprecated()));
-				AllFunctions.Sort((x, y) => String.Compare(x.Name, y.Name));
+				AllFunctions = AllFunctions.OrderBy(x => x.Name).ToList();
 
 				// Write all the specializations
 				if (TemplateSpecializations.Count > 0)
@@ -392,7 +392,7 @@ namespace APIDocTool
 				}
 
 				// Write all the variables
-				Writer.WriteListSection("variables", "Variables", "Name", "Description", Children.OfType<APIVariable>().Where(x => x.Protection != APIProtection.Private && !x.IsDeprecated()).OrderBy(x => x.Name).Select(x => x.GetListItem()));
+				Writer.WriteListSection("variables", "Variables", "Name", "Description", Children.OfType<APIVariable>().Where(x => x.Protection == APIProtection.Public && !x.IsDeprecated()).OrderBy(x => x.Name).Select(x => x.GetListItem()));
 
 				// Write all the constructors
 				if (!APIFunction.WriteListSection(Writer, "constructor", "Constructors", AllFunctions.Where(x => x.FunctionType == APIFunctionType.Constructor).OrderBy(x => x.LinkPath)) && HasAnyPrivateFunction(Name))
@@ -416,20 +416,27 @@ namespace APIDocTool
 
 				// Build a list of functions for each base record
 				List<APIFunction>[] AllBaseFunctions = AllBaseRecords.Select(x => new List<APIFunction>()).ToArray();
-				foreach(APIFunction Function in AllFunctions.Where(x =>x.FunctionType == APIFunctionType.Normal && !x.IsExecFunction()))
+				foreach (APIFunction Function in AllFunctions.Where(x => x.FunctionType == APIFunctionType.Normal && !x.IsExecFunction()))
 				{
 					int BaseRecordIdx = AllBaseRecords.IndexOf(Function.GetBaseImplementation().FindParent<APIRecord>());
 					AllBaseFunctions[Math.Max(0, BaseRecordIdx)].Add(Function);
 				}
 
+				// Write all the public functions in this class
+				APIFunction.WriteListSection(Writer, "methods", "Methods", AllBaseFunctions[0].Where(x => x.Protection == APIProtection.Public));
+
+				// Write the protected variables and functions
+				Writer.WriteListSection("protvariables", "Protected Variables", "Name", "Description", Children.OfType<APIVariable>().Where(x => x.Protection == APIProtection.Protected && !x.IsDeprecated()).OrderBy(x => x.Name).Select(x => x.GetListItem()));
+				APIFunction.WriteListSection(Writer, "protmethods", "Protected Methods", AllBaseFunctions[0].Where(x => x.Protection == APIProtection.Protected));
+
 				// Write the functions
-				for (int Idx = 0; Idx < AllBaseFunctions.Length; Idx++)
+				for (int Idx = 1; Idx < AllBaseFunctions.Length; Idx++)
 				{
 					List<APIFunction> BaseFunctions = AllBaseFunctions[Idx];
 					if (BaseFunctions.Count > 0)
 					{
 						string Id = String.Format("methods_{0}", Idx);
-						string Label = (Idx == 0) ? "Methods" : String.Format("Overridden from {0}", AllBaseRecords[Idx].Name);
+						string Label = String.Format("Overridden from {0}", AllBaseRecords[Idx].Name);
 						APIFunction.WriteListSection(Writer, Id, Label, AllBaseFunctions[Idx]);
 					}
 				}
