Blueprint API now displays images

diff --git a/Engine/Source/Editor/UnrealEd/Private/Commandlets/GenerateBlueprintAPICommandlet.cpp b/Engine/Source/Editor/UnrealEd/Private/Commandlets/GenerateBlueprintAPICommandlet.cpp
--- a/Engine/Source/Editor/UnrealEd/Private/Commandlets/GenerateBlueprintAPICommandlet.cpp
+++ b/Engine/Source/Editor/UnrealEd/Private/Commandlets/GenerateBlueprintAPICommandlet.cpp
@@ -6,6 +6,7 @@
 #include "K2ActionMenuBuilder.h"
 #include "EdGraphSchema_K2_Actions.h"
 #include "K2Node.h"
+#include "K2Node_CommutativeAssociativeBinaryOperator.h"
 #include "AnimationGraph.h"
 #include "AnimGraphNode_StateMachine.h"
 #include "BlueprintEditorUtils.h"
@@ -806,7 +807,7 @@ static void GenerateBlueprintAPIUtils::DumpActionMenuItem(uint32 Indent, FGraphA
 
 	ActionEntry += " : {";
 
-	const FString TooltipFieldLabel("\"Tooltip\"     : \"");
+	const FString TooltipFieldLabel("\"Tooltip\"      : \"");
 	const FString TooltipStr = PrimeAction->TooltipDescription.Replace(TEXT("\n"), *(IndentedNewline + BuildIndentString(TooltipFieldLabel.Len(), /*bUseSpaces =*/true)));
 
 	ActionEntry += IndentedNewline + TooltipFieldLabel + MakeJsonString(TooltipStr) + "\"";
@@ -819,9 +820,28 @@ static void GenerateBlueprintAPIUtils::DumpActionMenuItem(uint32 Indent, FGraphA
 		ActionList.OwnerOfTemporaries->AddNode(Node);
 		Node->AllocateDefaultPins();
 
+		if (Node->ShouldDrawCompact())
+		{
+			ActionEntry += IndentedNewline + "\"CompactTitle\" : \"" + MakeJsonString(Node->GetCompactNodeTitle().ToString()) + "\"";
+		}
+
+		if (Node->IsNodePure())
+		{
+			ActionEntry += IndentedNewline + "\"NodeType\"     : \"pure\"";
+		}
+		else
+		{
+			ActionEntry += IndentedNewline + "\"NodeType\"     : \"function\"";
+		}
+
+		if (Node->IsA<UK2Node_CommutativeAssociativeBinaryOperator>())
+		{
+			ActionEntry += IndentedNewline + "\"ShowAddPin\"   : \"true\"";
+		}
+
 		if (Node->Pins.Num() > 0)
 		{
-			ActionEntry += "," + IndentedNewline + "\"Pins\"        : {";
+			ActionEntry += "," + IndentedNewline + "\"Pins\"         : {";
 
 			const FString PinEntryIndentedNewline = "\n" + BuildIndentString(Indent+1);
 			const FString PinDetailsIndentedNewline = "\n" + BuildIndentString(Indent+2);
@@ -871,6 +891,16 @@ static void GenerateBlueprintAPIUtils::DumpActionMenuItem(uint32 Indent, FGraphA
 						ActionEntry += "," + PinDetailsIndentedNewline + "\"PinSubCategoryObject\" : \"" + MakeJsonString(Pin->PinType.PinSubCategoryObject->GetName()) + "\"";
 					}
 
+					if (!CastChecked<UEdGraphSchema_K2>(Node->GetSchema())->ShouldShowAssetPickerForPin(Pin))
+					{
+						ActionEntry += "," + PinDetailsIndentedNewline + "\"ShowAssetPicker\"      : \"false\"";
+					}
+
+					if (!Pin->DefaultValue.IsEmpty())
+					{
+						ActionEntry += "," + PinDetailsIndentedNewline + "\"DefaultValue\"         : \"" + MakeJsonString(Pin->DefaultValue) + "\"";
+					}
+
 					if (Pin->PinType.bIsArray)
 					{
 						ActionEntry += "," + PinDetailsIndentedNewline + "\"IsArray\"              : \"true\"";
diff --git a/Engine/Source/Programs/UnrealDocTool/APIDocTool/APIDocTool/APIAction.cs b/Engine/Source/Programs/UnrealDocTool/APIDocTool/APIDocTool/APIAction.cs
--- a/Engine/Source/Programs/UnrealDocTool/APIDocTool/APIDocTool/APIAction.cs
+++ b/Engine/Source/Programs/UnrealDocTool/APIDocTool/APIDocTool/APIAction.cs
@@ -9,6 +9,11 @@ namespace APIDocTool
 	class APIAction : APIPage
 	{
 		public readonly string Tooltip;
+		public readonly string CompactName;
+		public readonly string NodeType;
+
+		public readonly bool bShowAddPin;
+
 		private List<APIActionPin> Pins;
 
 		public APIAction(APIPage InParent, string InName, Dictionary<string, object> ActionProperties)
@@ -19,13 +24,26 @@ namespace APIDocTool
 			if (ActionProperties.TryGetValue("Tooltip", out Value))
 			{
 				Tooltip = String.Concat(((string)Value).Split('\n').Select(Line => Line.Trim() + '\n'));
-//				Tooltip = new string(Array.FindAll<char>(((string)Value).ToCharArray(), c => c != '\t'));
 			}
 			else
 			{
 				Tooltip = "";
 			}
 
+			if (ActionProperties.TryGetValue("CompactTitle", out Value))
+			{
+				CompactName = (string)Value;
+			}
+
+			if (ActionProperties.TryGetValue("NodeType", out Value))
+			{
+				NodeType = (string)Value;
+			}
+			else
+			{
+				NodeType = "function";
+			}
+
 			if (ActionProperties.TryGetValue("Pins", out Value))
 			{
 				Pins = new List<APIActionPin>();
@@ -35,6 +53,11 @@ namespace APIDocTool
 					Pins.Add(new APIActionPin(this, APIActionPin.GetPinName((Dictionary<string, object>)Pin.Value), (Dictionary<string, object>)Pin.Value));
 				}
 			}
+
+			if (ActionProperties.TryGetValue("ShowAddPin", out Value))
+			{
+				bShowAddPin = Convert.ToBoolean((string)Value);
+			}
 		}
 
 		private void WritePins(UdnWriter Writer, IEnumerable<APIActionPin> Pins)
@@ -45,6 +68,14 @@ namespace APIDocTool
 			}
 		}
 
+		private void WritePinObjects(UdnWriter Writer, IEnumerable<APIActionPin> Pins)
+		{
+			foreach(var Pin in Pins)
+			{
+				Pin.WriteObject(Writer, CompactName == null);
+			}
+		}
+
 		public override void WritePage(UdnManifest Manifest, string OutputPath)
 		{
 			using (UdnWriter Writer = new UdnWriter(OutputPath))
@@ -56,6 +87,38 @@ namespace APIDocTool
 
 				if (Pins != null && Pins.Count() > 0)
 				{
+					// Visualization of the node
+					Writer.EnterRegion("graph");
+					Writer.EnterObject("BlueprintNode");
+					if (CompactName != null)
+					{
+						Writer.WriteParamLiteral("type", "compact");
+						Writer.WriteParamLiteral("title", CompactName);
+					}
+					else
+					{
+						Writer.WriteParamLiteral("type", NodeType);
+						Writer.WriteParamLiteral("title", Name);
+					}
+					Writer.EnterParam("inputs");
+					WritePinObjects(Writer, Pins.Where(x => x.bInputPin));
+					Writer.LeaveParam();
+					Writer.EnterParam("outputs");
+					WritePinObjects(Writer, Pins.Where(x => !x.bInputPin));
+
+					if (bShowAddPin)
+					{
+						Writer.EnterObject("BlueprintPin");
+						Writer.WriteParamLiteral("type", "addpin");
+						Writer.WriteParamLiteral("id", "AddPin");
+						Writer.WriteParamLiteral("title", "Add pin");
+						Writer.LeaveObject();
+					}
+
+					Writer.LeaveParam();
+					Writer.LeaveObject();
+					Writer.LeaveRegion();
+
 					// Inputs
 					Writer.EnterSection("inputs", "Inputs");
 					Writer.WriteObject("MemberIconListHeadBlank");
diff --git a/Engine/Source/Programs/UnrealDocTool/APIDocTool/APIDocTool/APIActionPin.cs b/Engine/Source/Programs/UnrealDocTool/APIDocTool/APIDocTool/APIActionPin.cs
--- a/Engine/Source/Programs/UnrealDocTool/APIDocTool/APIDocTool/APIActionPin.cs
+++ b/Engine/Source/Programs/UnrealDocTool/APIDocTool/APIDocTool/APIActionPin.cs
@@ -11,6 +11,7 @@ namespace APIDocTool
 	{
 		public readonly string Tooltip;
 		public readonly string TypeText;
+		public readonly string DefaultValue;
 
 		public readonly string PinCategory;
 		public readonly string PinSubCategory;
@@ -20,6 +21,8 @@ namespace APIDocTool
 		public readonly bool bIsConst;
 		public readonly bool bIsReference;
 		public readonly bool bIsArray;
+		public readonly bool bShowAssetPicker;
+		public readonly bool bShowPinLabel;
 
 		public static string GetPinName(Dictionary<string, object> PinProperties)
 		{
@@ -47,6 +50,9 @@ namespace APIDocTool
 
 			Debug.Assert(InName == GetPinName(PinProperties));
 
+			// We have to check this again to set the show label property
+			bShowPinLabel = PinProperties.TryGetValue("Name", out Value);
+
 			bInputPin = (string)PinProperties["Direction"] == "input";
 			PinCategory = (string)PinProperties["PinCategory"];
 			TypeText = (string)PinProperties["TypeText"];
@@ -54,7 +60,6 @@ namespace APIDocTool
 			if (PinProperties.TryGetValue("Tooltip", out Value))
 			{
 				Tooltip = String.Concat(((string)Value).Split('\n').Select(Line => Line.Trim() + '\n'));
-//				Tooltip = new string(Array.FindAll<char>(((string)Value).ToCharArray(), c => c != '\t'));
 			}
 			else
 			{
@@ -78,7 +83,15 @@ namespace APIDocTool
 			{
 				PinSubCategoryObject = "";
 			}
-
+			
+			if (PinProperties.TryGetValue("DefaultValue", out Value))
+			{
+				DefaultValue = (string)Value;
+			}
+			else
+			{
+				DefaultValue = "";
+			}
 			if (PinProperties.TryGetValue("IsConst", out Value))
 			{
 				bIsConst = Convert.ToBoolean((string)Value);
@@ -93,6 +106,15 @@ namespace APIDocTool
 			{
 				bIsArray = Convert.ToBoolean((string)Value);
 			}
+
+			if (PinProperties.TryGetValue("ShowAssetPicker", out Value))
+			{
+				bShowAssetPicker = Convert.ToBoolean((string)Value);
+			}
+			else
+			{
+				bShowAssetPicker = true;
+			}
 		}
 
 		private string GetPinCategory()
@@ -112,10 +134,100 @@ namespace APIDocTool
 					return "rotator";
 				}
 			}
+			else if (PinCategory == "bool")
+			{
+				return "boolean";
+			}
+			else if (PinCategory == "int")
+			{
+				return "integer";
+			}
+			else if (PinCategory == "byte")
+			{
+				if (PinSubCategoryObject != "")
+				{
+					return "enum";
+				}
+			}
 
 			return PinCategory;
 		}
 
+		private string GetId()
+		{
+			return Name.Replace(' ', '_');
+		}
+
+		public void WriteObject(UdnWriter Writer, bool bDrawLabels)
+		{
+			string PinCategory = GetPinCategory();
+
+			Writer.EnterObject("BlueprintPin");
+			Writer.WriteParamLiteral("type", PinCategory + (bIsArray ? " array" : ""));
+			Writer.WriteParamLiteral("id", GetId());
+
+			if (bShowPinLabel && bDrawLabels)
+			{
+				Writer.WriteParamLiteral("title",  Name);
+			}
+
+			var DefaultValueElements = DefaultValue.Split(',');
+			switch(PinCategory)
+			{
+				case "byte":
+					Writer.WriteParamLiteral("value", (DefaultValueElements.Length > 0 ? DefaultValueElements[0] : "0"));
+					break;
+
+				case "class":
+					if (bShowAssetPicker)
+					{
+						Writer.WriteParamLiteral("value", "Select Class");
+					}
+					break;
+
+				case "float":
+					Writer.WriteParamLiteral("value", (DefaultValueElements.Length > 0 ? DefaultValueElements[0] : "0.0"));
+					break;
+
+				case "integer":
+					Writer.WriteParamLiteral("value", (DefaultValueElements.Length > 0 ? DefaultValueElements[0] : "0"));
+					break;
+
+				case "object":
+					if (bShowAssetPicker)
+					{
+						Writer.WriteParamLiteral("value", "Select Asset");
+					}
+					break;
+
+				case "rotator":
+					Writer.WriteParamLiteral("yaw", (DefaultValueElements.Length > 0 ? DefaultValueElements[0] : "0.0"));
+					Writer.WriteParamLiteral("pitch", (DefaultValueElements.Length > 1 ? DefaultValueElements[1] : "0.0"));
+					Writer.WriteParamLiteral("roll", (DefaultValueElements.Length > 2 ? DefaultValueElements[2] : "0.0"));
+					break;
+
+				case "struct":
+					if (PinSubCategoryObject == "Vector2D")
+					{
+						Writer.WriteParamLiteral("x", (DefaultValueElements.Length > 0 ? DefaultValueElements[0] : "0.0"));
+						Writer.WriteParamLiteral("y", (DefaultValueElements.Length > 1 ? DefaultValueElements[1] : "0.0"));
+					}
+					break;
+
+				case "vector":
+					Writer.WriteParamLiteral("x",(DefaultValueElements.Length > 0 ? DefaultValueElements[0] : "0.0"));
+					Writer.WriteParamLiteral("y",(DefaultValueElements.Length > 1 ? DefaultValueElements[1] : "0.0"));
+					Writer.WriteParamLiteral("z",(DefaultValueElements.Length > 2 ? DefaultValueElements[2] : "0.0"));
+					break;
+
+				default:
+					Writer.WriteParamLiteral("value", DefaultValue);
+					break;
+			}
+
+			Writer.LeaveObject();
+		}
+
 		public override void WritePage(UdnManifest Manifest, string OutputPath)
 		{
 			using (UdnWriter Writer = new UdnWriter(OutputPath))
@@ -147,7 +259,13 @@ namespace APIDocTool
 				ItemIcons.Add(Icons.VariablePinIcons[GetPinCategory()]);
 			}
 
-			Writer.WriteObject("ActionPinListItem", "icons", ItemIcons, "name", Name, "type", TypeText, "tooltip", Tooltip);
+			Writer.EnterObject("ActionPinListItem");
+			Writer.WriteParam("icons", ItemIcons);
+			Writer.WriteParamLiteral("name", Name);
+			Writer.WriteParamLiteral("type", TypeText);
+			Writer.WriteParam("tooltip", Tooltip);
+			Writer.WriteParamLiteral("id", GetId());
+			Writer.LeaveObject();
 		}
 	}
 }
diff --git a/Engine/Source/Programs/UnrealDocTool/APIDocTool/APIDocTool/Icons.cs b/Engine/Source/Programs/UnrealDocTool/APIDocTool/APIDocTool/Icons.cs
--- a/Engine/Source/Programs/UnrealDocTool/APIDocTool/APIDocTool/Icons.cs
+++ b/Engine/Source/Programs/UnrealDocTool/APIDocTool/APIDocTool/Icons.cs
@@ -149,13 +149,14 @@ namespace APIDocTool
 		//// Blueprint Variable type icons 
 		public static readonly Dictionary<string, Icon> VariablePinIcons = new Dictionary<string, Icon>
 		{
-			{ "bool", new Icon("%ROOT%/bp_api_pin_bool.png", "Boolean") },
+			{ "boolean", new Icon("%ROOT%/bp_api_pin_bool.png", "Boolean") },
 			{ "byte", new Icon("%ROOT%/bp_api_pin_byte.png", "Byte") },
 			{ "class", new Icon("%ROOT%/bp_api_pin_class.png", "Class") },
 			{ "delegate", new Icon("%ROOT%/bp_api_pin_delegate.png", "Delegate") },
+			{ "enum", new Icon("%ROOT%/bp_api_pin_byte.png", "Enum") },
 			{ "exec", new Icon("%ROOT%/bp_api_pin_exec.png", "Exec") },
 			{ "float", new Icon("%ROOT%/bp_api_pin_float.png", "Float") },
-			{ "int", new Icon("%ROOT%/bp_api_pin_int.png", "Integer") },
+			{ "integer", new Icon("%ROOT%/bp_api_pin_int.png", "Integer") },
 			{ "interface", new Icon("%ROOT%/bp_api_pin_interface.png", "Interface") },
 			{ "name", new Icon("%ROOT%/bp_api_pin_name.png", "Name") },
 			{ "object", new Icon("%ROOT%/bp_api_pin_object.png", "Object") },
@@ -171,11 +172,12 @@ namespace APIDocTool
 		//// Blueprint Array variable type icons 
 		public static readonly Dictionary<string, Icon> ArrayVariablePinIcons = new Dictionary<string, Icon>
 		{
-			{ "bool", new Icon("%ROOT%/bp_api_pin_array_bool.png", "Boolean Array") },
+			{ "boolean", new Icon("%ROOT%/bp_api_pin_array_bool.png", "Boolean Array") },
 			{ "byte", new Icon("%ROOT%/bp_api_pin_array_byte.png", "Byte Array") },
 			{ "class", new Icon("%ROOT%/bp_api_pin_array_class.png", "Class Array") },
+			{ "enum", new Icon("%ROOT%/bp_api_pin_array_byte.png", "Enum Array") },
 			{ "float", new Icon("%ROOT%/bp_api_pin_array_float.png", "Float Array") },
-			{ "int", new Icon("%ROOT%/bp_api_pin_array_int.png", "Integer Array") },
+			{ "integer", new Icon("%ROOT%/bp_api_pin_array_int.png", "Integer Array") },
 			{ "name", new Icon("%ROOT%/bp_api_pin_array_name.png", "Name Array") },
 			{ "object", new Icon("%ROOT%/bp_api_pin_array_object.png", "Object Array") },
 			{ "rotator", new Icon("%ROOT%/bp_api_pin_array_rotator.png", "Rotator Array") },
