Display pin value boxes with proper defaults for vector2D and linear color

diff --git a/Engine/Source/Programs/UnrealDocTool/APIDocTool/APIDocTool/APIActionPin.cs b/Engine/Source/Programs/UnrealDocTool/APIDocTool/APIDocTool/APIActionPin.cs
--- a/Engine/Source/Programs/UnrealDocTool/APIDocTool/APIDocTool/APIActionPin.cs
+++ b/Engine/Source/Programs/UnrealDocTool/APIDocTool/APIDocTool/APIActionPin.cs
@@ -134,6 +134,30 @@ namespace APIDocTool
 					return "rotator";
 				}
 			}
+			else if (PinCategory == "byte")
+			{
+				if (PinSubCategoryObject != "")
+				{
+					return "enum";
+				}
+			}
+
+			return PinCategory;
+		}
+
+		private string GetTypeText()
+		{
+			if (PinCategory == "struct")
+			{
+				if (PinSubCategoryObject == "Vector2D")
+				{
+					return "vector2d struct";
+				}
+				else if (PinSubCategoryObject == "LinearColor")
+				{
+					return "linearcolor";
+				}
+			}
 			else if (PinCategory == "bool")
 			{
 				return "boolean";
@@ -142,13 +166,6 @@ namespace APIDocTool
 			{
 				return "integer";
 			}
-			else if (PinCategory == "byte")
-			{
-				if (PinSubCategoryObject != "")
-				{
-					return "enum";
-				}
-			}
 
 			return PinCategory;
 		}
@@ -160,10 +177,8 @@ namespace APIDocTool
 
 		public void WriteObject(UdnWriter Writer, bool bDrawLabels)
 		{
-			string PinCategory = GetPinCategory();
-
 			Writer.EnterObject("BlueprintPin");
-			Writer.WriteParamLiteral("type", PinCategory + (bIsArray ? " array" : ""));
+			Writer.WriteParamLiteral("type", GetTypeText() + (bIsArray ? " array" : ""));
 			Writer.WriteParamLiteral("id", GetId());
 
 			if (bShowPinLabel && bDrawLabels)
@@ -172,10 +187,10 @@ namespace APIDocTool
 			}
 
 			var DefaultValueElements = DefaultValue.Split(',');
-			switch(PinCategory)
+			switch(GetPinCategory())
 			{
 				case "byte":
-					Writer.WriteParamLiteral("value", (DefaultValueElements.Length > 0 ? DefaultValueElements[0] : "0"));
+					Writer.WriteParamLiteral("value", (DefaultValueElements[0].Length > 0 ? DefaultValueElements[0] : "0"));
 					break;
 
 				case "class":
@@ -186,11 +201,11 @@ namespace APIDocTool
 					break;
 
 				case "float":
-					Writer.WriteParamLiteral("value", (DefaultValueElements.Length > 0 ? DefaultValueElements[0] : "0.0"));
+					Writer.WriteParamLiteral("value", (DefaultValueElements[0].Length > 0 ? DefaultValueElements[0] : "0.0"));
 					break;
 
 				case "integer":
-					Writer.WriteParamLiteral("value", (DefaultValueElements.Length > 0 ? DefaultValueElements[0] : "0"));
+					Writer.WriteParamLiteral("value", (DefaultValueElements[0].Length > 0 ? DefaultValueElements[0] : "0"));
 					break;
 
 				case "object":
@@ -201,21 +216,31 @@ namespace APIDocTool
 					break;
 
 				case "rotator":
-					Writer.WriteParamLiteral("yaw", (DefaultValueElements.Length > 0 ? DefaultValueElements[0] : "0.0"));
+					Writer.WriteParamLiteral("yaw", (DefaultValueElements[0].Length > 0 ? DefaultValueElements[0] : "0.0"));
 					Writer.WriteParamLiteral("pitch", (DefaultValueElements.Length > 1 ? DefaultValueElements[1] : "0.0"));
 					Writer.WriteParamLiteral("roll", (DefaultValueElements.Length > 2 ? DefaultValueElements[2] : "0.0"));
 					break;
 
 				case "struct":
-					if (PinSubCategoryObject == "Vector2D")
+					if (PinSubCategoryObject == "LinearColor")
+					{
+						float r = (DefaultValueElements[0].Length > 0 ? float.Parse(DefaultValueElements[0].Split('=')[1]) : 0) * 255;
+						float g = (DefaultValueElements.Length > 1 ? float.Parse(DefaultValueElements[1].Split('=')[1]) : 0) * 255;
+						float b = (DefaultValueElements.Length > 2 ? float.Parse(DefaultValueElements[2].Split('=')[1]) : 0) * 255;
+
+						Writer.WriteParamLiteral("r", ((int)r).ToString());
+						Writer.WriteParamLiteral("g", ((int)g).ToString());
+						Writer.WriteParamLiteral("b", ((int)b).ToString());
+					}
+					else if (PinSubCategoryObject == "Vector2D")
 					{
-						Writer.WriteParamLiteral("x", (DefaultValueElements.Length > 0 ? DefaultValueElements[0] : "0.0"));
+						Writer.WriteParamLiteral("x", (DefaultValueElements[0].Length > 0 ? DefaultValueElements[0] : "0.0"));
 						Writer.WriteParamLiteral("y", (DefaultValueElements.Length > 1 ? DefaultValueElements[1] : "0.0"));
 					}
 					break;
 
 				case "vector":
-					Writer.WriteParamLiteral("x",(DefaultValueElements.Length > 0 ? DefaultValueElements[0] : "0.0"));
+					Writer.WriteParamLiteral("x", (DefaultValueElements[0].Length > 0 ? DefaultValueElements[0] : "0.0"));
 					Writer.WriteParamLiteral("y",(DefaultValueElements.Length > 1 ? DefaultValueElements[1] : "0.0"));
 					Writer.WriteParamLiteral("z",(DefaultValueElements.Length > 2 ? DefaultValueElements[2] : "0.0"));
 					break;
diff --git a/Engine/Source/Programs/UnrealDocTool/APIDocTool/APIDocTool/Icons.cs b/Engine/Source/Programs/UnrealDocTool/APIDocTool/APIDocTool/Icons.cs
--- a/Engine/Source/Programs/UnrealDocTool/APIDocTool/APIDocTool/Icons.cs
+++ b/Engine/Source/Programs/UnrealDocTool/APIDocTool/APIDocTool/Icons.cs
@@ -149,14 +149,14 @@ namespace APIDocTool
 		//// Blueprint Variable type icons 
 		public static readonly Dictionary<string, Icon> VariablePinIcons = new Dictionary<string, Icon>
 		{
-			{ "boolean", new Icon("%ROOT%/bp_api_pin_bool.png", "Boolean") },
+			{ "bool", new Icon("%ROOT%/bp_api_pin_bool.png", "Boolean") },
 			{ "byte", new Icon("%ROOT%/bp_api_pin_byte.png", "Byte") },
 			{ "class", new Icon("%ROOT%/bp_api_pin_class.png", "Class") },
 			{ "delegate", new Icon("%ROOT%/bp_api_pin_delegate.png", "Delegate") },
 			{ "enum", new Icon("%ROOT%/bp_api_pin_byte.png", "Enum") },
 			{ "exec", new Icon("%ROOT%/bp_api_pin_exec.png", "Exec") },
 			{ "float", new Icon("%ROOT%/bp_api_pin_float.png", "Float") },
-			{ "integer", new Icon("%ROOT%/bp_api_pin_int.png", "Integer") },
+			{ "int", new Icon("%ROOT%/bp_api_pin_int.png", "Integer") },
 			{ "interface", new Icon("%ROOT%/bp_api_pin_interface.png", "Interface") },
 			{ "name", new Icon("%ROOT%/bp_api_pin_name.png", "Name") },
 			{ "object", new Icon("%ROOT%/bp_api_pin_object.png", "Object") },
@@ -172,12 +172,12 @@ namespace APIDocTool
 		//// Blueprint Array variable type icons 
 		public static readonly Dictionary<string, Icon> ArrayVariablePinIcons = new Dictionary<string, Icon>
 		{
-			{ "boolean", new Icon("%ROOT%/bp_api_pin_array_bool.png", "Boolean Array") },
+			{ "bool", new Icon("%ROOT%/bp_api_pin_array_bool.png", "Boolean Array") },
 			{ "byte", new Icon("%ROOT%/bp_api_pin_array_byte.png", "Byte Array") },
 			{ "class", new Icon("%ROOT%/bp_api_pin_array_class.png", "Class Array") },
 			{ "enum", new Icon("%ROOT%/bp_api_pin_array_byte.png", "Enum Array") },
 			{ "float", new Icon("%ROOT%/bp_api_pin_array_float.png", "Float Array") },
-			{ "integer", new Icon("%ROOT%/bp_api_pin_array_int.png", "Integer Array") },
+			{ "int", new Icon("%ROOT%/bp_api_pin_array_int.png", "Integer Array") },
 			{ "name", new Icon("%ROOT%/bp_api_pin_array_name.png", "Name Array") },
 			{ "object", new Icon("%ROOT%/bp_api_pin_array_object.png", "Object Array") },
 			{ "rotator", new Icon("%ROOT%/bp_api_pin_array_rotator.png", "Rotator Array") },
