Rename 'functions' heading to 'methods'.

diff --git a/Engine/Source/Programs/UnrealDocTool/APIDocTool/APIDocTool/APIRecord.cs b/Engine/Source/Programs/UnrealDocTool/APIDocTool/APIDocTool/APIRecord.cs
--- a/Engine/Source/Programs/UnrealDocTool/APIDocTool/APIDocTool/APIRecord.cs
+++ b/Engine/Source/Programs/UnrealDocTool/APIDocTool/APIDocTool/APIRecord.cs
@@ -428,8 +428,8 @@ namespace APIDocTool
 					List<APIFunction> BaseFunctions = AllBaseFunctions[Idx];
 					if (BaseFunctions.Count > 0)
 					{
-						string Id = String.Format("methods_{0}", Idx);
-						string Label = (Idx == 0) ? "Methods" : String.Format("Overridden from {0}", AllBaseRecords[Idx].Name);
+						string Id = String.Format("functions_{0}", Idx);
+						string Label = (Idx == 0) ? "Functions" : String.Format("Overridden from {0}", AllBaseRecords[Idx].Name);
 						APIFunction.WriteListSection(Writer, Id, Label, AllBaseFunctions[Idx]);
 					}
 				}
