#UE4 Docs: fixed javascript not being included in VS plugin preview

diff --git a/Engine/Source/Programs/UnrealDocTool/MarkdownMode/ToolWindow/PreviewWindowUpdateListener.cs b/Engine/Source/Programs/UnrealDocTool/MarkdownMode/ToolWindow/PreviewWindowUpdateListener.cs
--- a/Engine/Source/Programs/UnrealDocTool/MarkdownMode/ToolWindow/PreviewWindowUpdateListener.cs
+++ b/Engine/Source/Programs/UnrealDocTool/MarkdownMode/ToolWindow/PreviewWindowUpdateListener.cs
@@ -157,7 +157,8 @@ namespace MarkdownMode
                     anchorPosition, "<span name=\"MARKDOWNANCHORNOTUSEDELSEWHERE\"></span>");
             }
 
-            html.AppendLine(RemoveScript.Replace(newHtml, ""));
+            //html.AppendLine(RemoveScript.Replace(newHtml, ""));
+            html.AppendLine(newHtml);
 
             LogMarkdownErrors(document.FilePath, results.Errors, log);
 
diff --git a/Engine/Source/Programs/UnrealDocTool/MarkdownMode/source.extension.vsixmanifest b/Engine/Source/Programs/UnrealDocTool/MarkdownMode/source.extension.vsixmanifest
--- a/Engine/Source/Programs/UnrealDocTool/MarkdownMode/source.extension.vsixmanifest
+++ b/Engine/Source/Programs/UnrealDocTool/MarkdownMode/source.extension.vsixmanifest
@@ -1,7 +1,7 @@
 <?xml version="1.0" encoding="utf-8"?>
 <PackageManifest Version="2.0.0" xmlns="http://schemas.microsoft.com/developer/vsx-schema/2011" xmlns:d="http://schemas.microsoft.com/developer/vsx-schema-design/2011">
   <Metadata>
-    <Identity Id="Epic MarkdownMode" Version="0.222" Language="en-US" Publisher="Epic Games" />
+    <Identity Id="Epic MarkdownMode" Version="0.223" Language="en-US" Publisher="Epic Games" />
     <DisplayName>Epic MarkdownMode</DisplayName>
     <Description xml:space="preserve">A language service for editing Markdown files.</Description>
     <Icon>markdown-icon.png</Icon>
