#UE4 Docs: updated related pages metadata to be more useful

diff --git a/Engine/Source/Programs/UnrealDocTool/UnrealDocTool/MarkdownSharp/Markdown.cs b/Engine/Source/Programs/UnrealDocTool/UnrealDocTool/MarkdownSharp/Markdown.cs
--- a/Engine/Source/Programs/UnrealDocTool/UnrealDocTool/MarkdownSharp/Markdown.cs
+++ b/Engine/Source/Programs/UnrealDocTool/UnrealDocTool/MarkdownSharp/Markdown.cs
@@ -1130,10 +1130,18 @@ namespace MarkdownSharp
                         title = data.ProcessedDocumentCache.CurrentFileDocument.Metadata.DocumentTitle,
                         metadata = metadata,
                         translatedPages = translatedPageLinks,
-                        relatedPages = Templates.RelatedPages.Render(Hash.FromAnonymousObject(
+                        relatedPagesMenu = Templates.RelatedPages.Render(Hash.FromAnonymousObject(
                             new {
                                 relatedPages = data.ProcessedDocumentCache.Metadata.RelatedLinks,
-                                relatedPagesCount = data.ProcessedDocumentCache.Metadata.RelatedLinks.Count
+                                relatedPagesCount = data.ProcessedDocumentCache.Metadata.RelatedLinks.Count,
+                                quickjump = ""
+                            })),
+                        relatedPages = Templates.RelatedPages.Render(Hash.FromAnonymousObject(
+                            new
+                            {
+                                relatedPages = data.ProcessedDocumentCache.Metadata.RelatedLinks,
+                                relatedPagesCount = data.ProcessedDocumentCache.Metadata.RelatedLinks.Count,
+                                relativeHtmlPath = data.CurrentFolderDetails.RelativeHTMLPath
                             })),
                         errors = ThisIsPreview
                                      ? Templates.ErrorDetails.Render(Hash.FromAnonymousObject(
