Stop returning from the function that writes the blueprint API JSON archive if the editor commandlet succeeded.

diff --git a/Engine/Source/Programs/UnrealDocTool/APIDocTool/APIDocTool/Program.cs b/Engine/Source/Programs/UnrealDocTool/APIDocTool/APIDocTool/Program.cs
--- a/Engine/Source/Programs/UnrealDocTool/APIDocTool/APIDocTool/Program.cs
+++ b/Engine/Source/Programs/UnrealDocTool/APIDocTool/APIDocTool/Program.cs
@@ -714,7 +714,10 @@ namespace APIDocTool
 						JsonExportProcess.BeginOutputReadLine();
 						JsonExportProcess.BeginErrorReadLine();
 						JsonExportProcess.WaitForExit();
-						return JsonExportProcess.ExitCode == 0;
+						if(JsonExportProcess.ExitCode != 0)
+						{
+							return false;
+						}
 					}
 					catch (Exception Ex)
 					{
