Merging from 4.5 branch:

Fix well known struct pin types and enums not displaying correctly on the generated node picture

diff --git a/Engine/Source/Programs/UnrealDocTool/APIDocTool/APIDocTool/APIActionPin.cs b/Engine/Source/Programs/UnrealDocTool/APIDocTool/APIDocTool/APIActionPin.cs
--- a/Engine/Source/Programs/UnrealDocTool/APIDocTool/APIDocTool/APIActionPin.cs
+++ b/Engine/Source/Programs/UnrealDocTool/APIDocTool/APIDocTool/APIActionPin.cs
@@ -149,6 +149,18 @@ namespace APIDocTool
 		{
 			if (PinCategory == "struct")
 			{
+				if (PinSubCategoryObject == "Vector")
+				{
+					return "vector";
+				}
+				else if (PinSubCategoryObject == "Transform")
+				{
+					return "transform";
+				}
+				else if (PinSubCategoryObject == "Rotator")
+				{
+					return "rotator";
+				}
 				if (PinSubCategoryObject == "Vector2D")
 				{
 					return "vector2d struct";
@@ -166,6 +178,13 @@ namespace APIDocTool
 			{
 				return "integer";
 			}
+			else if (PinCategory == "byte")
+			{
+				if (PinSubCategoryObject != "")
+				{
+					return "enum";
+				}
+			}
 
 			return PinCategory;
 		}
