Merging from 4.5 branch:

Switch over to using new pin icon method
Fix default values for Vector2D

diff --git a/Engine/Source/Programs/UnrealDocTool/APIDocTool/APIDocTool/APIActionPin.cs b/Engine/Source/Programs/UnrealDocTool/APIDocTool/APIDocTool/APIActionPin.cs
--- a/Engine/Source/Programs/UnrealDocTool/APIDocTool/APIDocTool/APIActionPin.cs
+++ b/Engine/Source/Programs/UnrealDocTool/APIDocTool/APIDocTool/APIActionPin.cs
@@ -117,34 +117,6 @@ namespace APIDocTool
 			}
 		}
 
-		private string GetPinCategory()
-		{
-			if (PinCategory == "struct")
-			{
-				if (PinSubCategoryObject == "Vector")
-				{
-					return "vector";
-				}
-				else if (PinSubCategoryObject == "Transform")
-				{
-					return "transform";
-				}
-				else if (PinSubCategoryObject == "Rotator")
-				{
-					return "rotator";
-				}
-			}
-			else if (PinCategory == "byte")
-			{
-				if (PinSubCategoryObject != "")
-				{
-					return "enum";
-				}
-			}
-
-			return PinCategory;
-		}
-
 		private string GetTypeText()
 		{
 			if (PinCategory == "struct")
@@ -205,8 +177,8 @@ namespace APIDocTool
 				Writer.WriteParamLiteral("title",  Name);
 			}
 
-			var DefaultValueElements = DefaultValue.Split(',');
-			switch(GetPinCategory())
+			var DefaultValueElements = DefaultValue.Replace('(',' ').Replace(')',' ').Split(',');
+			switch (GetTypeText())
 			{
 				case "byte":
 					Writer.WriteParamLiteral("value", (DefaultValueElements[0].Length > 0 ? DefaultValueElements[0] : "0"));
@@ -240,22 +212,19 @@ namespace APIDocTool
 					Writer.WriteParamLiteral("roll", (DefaultValueElements.Length > 2 ? DefaultValueElements[2] : "0.0"));
 					break;
 
-				case "struct":
-					if (PinSubCategoryObject == "LinearColor")
-					{
-						float r = (DefaultValueElements[0].Length > 0 ? float.Parse(DefaultValueElements[0].Split('=')[1]) : 0) * 255;
-						float g = (DefaultValueElements.Length > 1 ? float.Parse(DefaultValueElements[1].Split('=')[1]) : 0) * 255;
-						float b = (DefaultValueElements.Length > 2 ? float.Parse(DefaultValueElements[2].Split('=')[1]) : 0) * 255;
+				case "linearcolor":
+					float r = (DefaultValueElements[0].Length > 0 ? float.Parse(DefaultValueElements[0].Split('=')[1]) : 0) * 255;
+					float g = (DefaultValueElements.Length > 1 ? float.Parse(DefaultValueElements[1].Split('=')[1]) : 0) * 255;
+					float b = (DefaultValueElements.Length > 2 ? float.Parse(DefaultValueElements[2].Split('=')[1]) : 0) * 255;
 
-						Writer.WriteParamLiteral("r", ((int)r).ToString());
-						Writer.WriteParamLiteral("g", ((int)g).ToString());
-						Writer.WriteParamLiteral("b", ((int)b).ToString());
-					}
-					else if (PinSubCategoryObject == "Vector2D")
-					{
-						Writer.WriteParamLiteral("x", (DefaultValueElements[0].Length > 0 ? DefaultValueElements[0] : "0.0"));
-						Writer.WriteParamLiteral("y", (DefaultValueElements.Length > 1 ? DefaultValueElements[1] : "0.0"));
-					}
+					Writer.WriteParamLiteral("r", ((int)r).ToString());
+					Writer.WriteParamLiteral("g", ((int)g).ToString());
+					Writer.WriteParamLiteral("b", ((int)b).ToString());
+					break;
+
+				case "vector2d struct":
+					Writer.WriteParamLiteral("x", (DefaultValueElements[0].Length > 0 ? DefaultValueElements[0].Split('=')[1] : "0.0"));
+					Writer.WriteParamLiteral("y", (DefaultValueElements.Length > 1 ? DefaultValueElements[1].Split('=')[1] : "0.0"));
 					break;
 
 				case "vector":
@@ -292,19 +261,25 @@ namespace APIDocTool
 
 		public void WritePin(UdnWriter Writer)
 		{
-			// Get all the icons
-			List<Icon> ItemIcons = new List<Icon>();
-			if (bIsArray)
+			Writer.EnterObject("ActionPinListItem");
+			
+			Writer.EnterParam("icons");
+
+			var PinType = GetTypeText();
+			if (PinType == "exec")
+			{
+				Writer.WriteRegion("input_exec","");
+			}
+			else if (bIsArray)
 			{
-				ItemIcons.Add(Icons.ArrayVariablePinIcons[GetPinCategory()]);
+				Writer.WriteRegion("input_array " + PinType, "");
 			}
 			else
 			{
-				ItemIcons.Add(Icons.VariablePinIcons[GetPinCategory()]);
+				Writer.WriteRegion("input_variable " + PinType, "");
 			}
+			Writer.LeaveParam();
 
-			Writer.EnterObject("ActionPinListItem");
-			Writer.WriteParam("icons", ItemIcons);
 			Writer.WriteParamLiteral("name", Name);
 			Writer.WriteParamLiteral("type", TypeText);
 			Writer.WriteParam("tooltip", Tooltip);
