Change directory that blueprint api use to look for images when building CHM.

diff --git a/Engine/Source/Programs/UnrealDocTool/APIDocTool/APIDocTool/Program.cs b/Engine/Source/Programs/UnrealDocTool/APIDocTool/APIDocTool/Program.cs
--- a/Engine/Source/Programs/UnrealDocTool/APIDocTool/APIDocTool/Program.cs
+++ b/Engine/Source/Programs/UnrealDocTool/APIDocTool/APIDocTool/Program.cs
@@ -1116,7 +1116,7 @@ namespace APIDocTool
 				// Find all the HTML files
 				if (bBlueprint)
 				{
-					Utility.FindRelativeContents(BaseHtmlDir, "Images\\bp_api*", false, FilePaths, DirectoryPaths);
+					Utility.FindRelativeContents(BaseHtmlDir, "Images\\Include\\bp_api*", false, FilePaths, DirectoryPaths);
 					Utility.FindRelativeContents(BaseHtmlDir, "INT\\BlueprintAPI\\*.html", true, FilePaths, DirectoryPaths);
 					CreateChm(ChmCompilerPath, ChmFileName, "UE4 Blueprint API Documentation", "INT\\BlueprintAPI\\index.html", ContentsFileName, IndexFileName, BaseHtmlDir, FilePaths);
 				}
