Fix include path to BP doc images.

diff --git a/Engine/Source/Programs/UnrealDocTool/APIDocTool/APIDocTool/Program.cs b/Engine/Source/Programs/UnrealDocTool/APIDocTool/APIDocTool/Program.cs
--- a/Engine/Source/Programs/UnrealDocTool/APIDocTool/APIDocTool/Program.cs
+++ b/Engine/Source/Programs/UnrealDocTool/APIDocTool/APIDocTool/Program.cs
@@ -1116,7 +1116,7 @@ namespace APIDocTool
 				// Find all the HTML files
 				if (bBlueprint)
 				{
-					Utility.FindRelativeContents(BaseHtmlDir, "Images\\Include\\bp_api*", false, FilePaths, DirectoryPaths);
+					Utility.FindRelativeContents(BaseHtmlDir, "Include\\Images\\bp_api*", false, FilePaths, DirectoryPaths);
 					Utility.FindRelativeContents(BaseHtmlDir, "INT\\BlueprintAPI\\*.html", true, FilePaths, DirectoryPaths);
 					CreateChm(ChmCompilerPath, ChmFileName, "UE4 Blueprint API Documentation", "INT\\BlueprintAPI\\index.html", ContentsFileName, IndexFileName, BaseHtmlDir, FilePaths);
 				}
