Don't list certain files twice in the [FILES] section of the project file

diff --git a/Engine/Source/Programs/UnrealDocTool/APIDocTool/APIDocTool/Program.cs b/Engine/Source/Programs/UnrealDocTool/APIDocTool/APIDocTool/Program.cs
--- a/Engine/Source/Programs/UnrealDocTool/APIDocTool/APIDocTool/Program.cs
+++ b/Engine/Source/Programs/UnrealDocTool/APIDocTool/APIDocTool/Program.cs
@@ -1116,7 +1116,6 @@ namespace APIDocTool
 				// Find all the HTML files
 				if (bBlueprint)
 				{
-					Utility.FindRelativeContents(BaseHtmlDir, "Include\\Images\\bp_api*", false, FilePaths, DirectoryPaths);
 					Utility.FindRelativeContents(BaseHtmlDir, "INT\\BlueprintAPI\\*.html", true, FilePaths, DirectoryPaths);
 					CreateChm(ChmCompilerPath, ChmFileName, "UE4 Blueprint API Documentation", "INT\\BlueprintAPI\\index.html", ContentsFileName, IndexFileName, BaseHtmlDir, FilePaths);
 				}
