Fix for cyclic metadata references

 - made crumbs/related page references not do full processing of the referenced page's metadata

diff --git a/Engine/Source/Programs/UnrealDocTool/UnrealDocTool/MarkdownSharp/EpicMarkdown/EMDocument.cs b/Engine/Source/Programs/UnrealDocTool/UnrealDocTool/MarkdownSharp/EpicMarkdown/EMDocument.cs
--- a/Engine/Source/Programs/UnrealDocTool/UnrealDocTool/MarkdownSharp/EpicMarkdown/EMDocument.cs
+++ b/Engine/Source/Programs/UnrealDocTool/UnrealDocTool/MarkdownSharp/EpicMarkdown/EMDocument.cs
@@ -93,7 +93,7 @@ namespace MarkdownSharp.EpicMarkdown
 
         public static EMDocument Empty = new EMDocument();
 
-        public void Parse(string text, bool preprocessOnly = false)
+        public void Parse(string text, bool preprocessOnly = false, bool full = true)
         {
             Text = text;
 
@@ -104,7 +104,7 @@ namespace MarkdownSharp.EpicMarkdown
                 relativePath = GetAbsoluteMarkdownPath();
             }
 
-            preprocessedText = TransformationData.Markdown.Preprocessor.Preprocess(text, TransformationData, preprocessedData, relativePath);
+            preprocessedText = TransformationData.Markdown.Preprocessor.Preprocess(text, TransformationData, preprocessedData, relativePath, full);
             publish = Markdown.CheckAvailability(TransformationData);
 
             if (!publish)
diff --git a/Engine/Source/Programs/UnrealDocTool/UnrealDocTool/MarkdownSharp/Preprocessor/MetadataManager.cs b/Engine/Source/Programs/UnrealDocTool/UnrealDocTool/MarkdownSharp/Preprocessor/MetadataManager.cs
--- a/Engine/Source/Programs/UnrealDocTool/UnrealDocTool/MarkdownSharp/Preprocessor/MetadataManager.cs
+++ b/Engine/Source/Programs/UnrealDocTool/UnrealDocTool/MarkdownSharp/Preprocessor/MetadataManager.cs
@@ -39,7 +39,7 @@ namespace MarkdownSharp.Preprocessor
         /// </summary>
         /// <param name="match"></param>
         /// <returns></returns>
-        private string ParseMetadataRow(Match match, TransformationData data)
+        private string ParseMetadataRow(Match match, TransformationData data, bool full = true)
         {
             var metaValue = match.Groups["MetadataValue"].Value;
             var metaDataCategory = match.Groups["MetadataKey"].Value;
@@ -69,7 +69,7 @@ namespace MarkdownSharp.Preprocessor
                     CrumbsLinks.Add(metaValue);
                 }
 
-                if (metaDataCategoryLowerCase == "related")
+                if (metaDataCategoryLowerCase == "related" && full)
                 {
                     RelatedLinks.Add(data.Markdown.ProcessRelated(metaValue, data));
                 }
@@ -115,19 +115,19 @@ namespace MarkdownSharp.Preprocessor
             return "";
         }
 
-        private string ParseMetadataMatches(Match match, TransformationData data)
+        private string ParseMetadataMatches(Match match, TransformationData data, bool full = true)
         {
             var metaDataRows = Regex.Split(match.Groups[0].Value, "\n");
 
             foreach (var currentMetaDataRow in metaDataRows)
             {
-                MetadataRow.Replace(currentMetaDataRow, m => ParseMetadataRow(m, data));
+                MetadataRow.Replace(currentMetaDataRow, m => ParseMetadataRow(m, data, full));
             }
 
             return "";
         }
 
-        public string ParseMetadata(string text, TransformationData data, PreprocessedData preprocessedData)
+        public string ParseMetadata(string text, TransformationData data, PreprocessedData preprocessedData, bool full = true)
         {
             var map = preprocessedData.TextMap;
 
@@ -136,7 +136,7 @@ namespace MarkdownSharp.Preprocessor
 
             var changes = new List<PreprocessingTextChange>();
 
-            text = Preprocessor.Replace(MetadataPattern, text, match => ParseMetadataMatches(match, data), null, changes);
+            text = Preprocessor.Replace(MetadataPattern, text, match => ParseMetadataMatches(match, data, full), null, changes);
 
             Preprocessor.AddChangesToBounds(map, changes, preprocessedData, PreprocessedTextType.Metadata);
 
diff --git a/Engine/Source/Programs/UnrealDocTool/UnrealDocTool/MarkdownSharp/Preprocessor/Preprocessor.cs b/Engine/Source/Programs/UnrealDocTool/UnrealDocTool/MarkdownSharp/Preprocessor/Preprocessor.cs
--- a/Engine/Source/Programs/UnrealDocTool/UnrealDocTool/MarkdownSharp/Preprocessor/Preprocessor.cs
+++ b/Engine/Source/Programs/UnrealDocTool/UnrealDocTool/MarkdownSharp/Preprocessor/Preprocessor.cs
@@ -75,7 +75,7 @@ namespace MarkdownSharp.Preprocessor
             }
         }
 
-        public string Preprocess(string text, TransformationData transData, PreprocessedData data, string relativePathToLinkedFile = null)
+        public string Preprocess(string text, TransformationData transData, PreprocessedData data, string relativePathToLinkedFile = null, bool full = true)
         {
             var relPath = relativePathToLinkedFile ?? transData.CurrentFolderDetails.CurrentFolderFromMarkdownAsTopLeaf;
 
@@ -87,7 +87,7 @@ namespace MarkdownSharp.Preprocessor
 
             text = DoCodeSpans(text, data.TextMap);
 
-            text = data.Metadata.ParseMetadata(text, transData, data);
+            text = data.Metadata.ParseMetadata(text, transData, data, full);
 
             text = filter.FilterOut(text, transData, data);
 
diff --git a/Engine/Source/Programs/UnrealDocTool/UnrealDocTool/MarkdownSharp/Preprocessor/ProcessedDocumentCache.cs b/Engine/Source/Programs/UnrealDocTool/UnrealDocTool/MarkdownSharp/Preprocessor/ProcessedDocumentCache.cs
--- a/Engine/Source/Programs/UnrealDocTool/UnrealDocTool/MarkdownSharp/Preprocessor/ProcessedDocumentCache.cs
+++ b/Engine/Source/Programs/UnrealDocTool/UnrealDocTool/MarkdownSharp/Preprocessor/ProcessedDocumentCache.cs
@@ -60,7 +60,7 @@ namespace MarkdownSharp.Preprocessor
                 var text = FileHelper.SafeReadAllText(path);
 
                 document = new EMDocument(path, cache.TransformationData);
-                document.Parse(text, true);
+                document.Parse(text,true,false);
             }
 
             private void FailWhenDisposed()
