Prevent truncating brief descriptions at common abbreviations, such as 'eg.' or 'ie.'

diff --git a/Engine/Source/Programs/UnrealDocTool/APIDocTool/APIDocTool/Utility.cs b/Engine/Source/Programs/UnrealDocTool/APIDocTool/APIDocTool/Utility.cs
--- a/Engine/Source/Programs/UnrealDocTool/APIDocTool/APIDocTool/Utility.cs
+++ b/Engine/Source/Programs/UnrealDocTool/APIDocTool/APIDocTool/Utility.cs
@@ -273,6 +273,12 @@ namespace APIDocTool
 
 		public static string GetBriefDescription(string Description)
 		{
+			// If it's short anyway, just use the full description
+			if(Description.Length < 100)
+			{
+				return Description;
+			}
+
 			// Find the first newline, or period that's followed by a space and not within any type of paren
 			int Nesting = 0;
 			for(int Idx = 0; Idx < Description.Length; Idx++)
@@ -302,6 +308,16 @@ namespace APIDocTool
 							return Description.Substring(0, Idx + 1);
 						}
 						break;
+					case ' ':
+						if(MatchSubstring(Description, Idx + 1, "eg.") || MatchSubstring(Description, Idx + 1, "ie."))
+						{
+							Idx += 3;
+						}
+						else if(MatchSubstring(Description, Idx + 1, "e.g.") || MatchSubstring(Description, Idx + 1, "i.e.") || MatchSubstring(Description, Idx + 1, "wrt."))
+						{
+							Idx += 4;
+						}
+						break;
 				}
 			}
 
@@ -309,6 +325,19 @@ namespace APIDocTool
 			return Description;
 		}
 
+		public static bool MatchSubstring(string Text, int StartIdx, string Substring)
+		{
+			if(StartIdx + Substring.Length > Text.Length)
+			{
+				return false;
+			}
+			for(int Idx = 0; Idx < Substring.Length; Idx++)
+			{
+				if(Text[StartIdx + Idx] != Substring[Idx]) return false;
+			}
+			return true;
+		}
+
 		public static string ClampStringLength(string Text, int MaxLength)
 		{
 			if (Text.Length <= MaxLength)
