Attempt to sneak under the page limit for building CHM files by only writing variable pages that won't fit on index pages.

diff --git a/Engine/Source/Programs/UnrealDocTool/APIDocTool/APIDocTool/APIFilter.cs b/Engine/Source/Programs/UnrealDocTool/APIDocTool/APIDocTool/APIFilter.cs
--- a/Engine/Source/Programs/UnrealDocTool/APIDocTool/APIDocTool/APIFilter.cs
+++ b/Engine/Source/Programs/UnrealDocTool/APIDocTool/APIDocTool/APIFilter.cs
@@ -118,7 +118,7 @@ namespace APIDocTool
 
 				// Write the variable list
 				Writer.EnterTag("[PARAM:variables]");
-				Writer.WriteListSection("variables", "Variables", "Name", "Description", FilteredMembers.OfType<APIVariable>().Select(x => x.GetListItem()));
+				APIVariable.WriteListSection(Writer, "variables", "Variables", FilteredMembers.OfType<APIVariable>());
 				Writer.LeaveTag("[/PARAM]");
 
 				// Close the module template
diff --git a/Engine/Source/Programs/UnrealDocTool/APIDocTool/APIDocTool/APIPage.cs b/Engine/Source/Programs/UnrealDocTool/APIDocTool/APIDocTool/APIPage.cs
--- a/Engine/Source/Programs/UnrealDocTool/APIDocTool/APIDocTool/APIPage.cs
+++ b/Engine/Source/Programs/UnrealDocTool/APIDocTool/APIDocTool/APIPage.cs
@@ -86,6 +86,11 @@ namespace APIDocTool
 			return VisitedSet;
 		}
 
+		public virtual bool ShouldOutputPage()
+		{
+			return true;
+		}
+
 		public virtual SitemapNode CreateSitemapNode()
 		{
 			return new SitemapNode(Name, SitemapLinkPath);
diff --git a/Engine/Source/Programs/UnrealDocTool/APIDocTool/APIDocTool/APIRecord.cs b/Engine/Source/Programs/UnrealDocTool/APIDocTool/APIDocTool/APIRecord.cs
--- a/Engine/Source/Programs/UnrealDocTool/APIDocTool/APIDocTool/APIRecord.cs
+++ b/Engine/Source/Programs/UnrealDocTool/APIDocTool/APIDocTool/APIRecord.cs
@@ -392,7 +392,7 @@ namespace APIDocTool
 				}
 
 				// Write all the variables
-				Writer.WriteListSection("variables", "Variables", "Name", "Description", Children.OfType<APIVariable>().Where(x => x.Protection != APIProtection.Private && !x.IsDeprecated()).OrderBy(x => x.Name).Select(x => x.GetListItem()));
+				APIVariable.WriteListSection(Writer, "variables", "Variables", Children.OfType<APIVariable>().Where(x => x.Protection != APIProtection.Private && !x.IsDeprecated()).OrderBy(x => x.Name));
 
 				// Write all the constructors
 				if (!APIFunction.WriteListSection(Writer, "constructor", "Constructors", AllFunctions.Where(x => x.FunctionType == APIFunctionType.Constructor).OrderBy(x => x.LinkPath)) && HasAnyPrivateFunction(Name))
diff --git a/Engine/Source/Programs/UnrealDocTool/APIDocTool/APIDocTool/APIVariable.cs b/Engine/Source/Programs/UnrealDocTool/APIDocTool/APIDocTool/APIVariable.cs
--- a/Engine/Source/Programs/UnrealDocTool/APIDocTool/APIDocTool/APIVariable.cs
+++ b/Engine/Source/Programs/UnrealDocTool/APIDocTool/APIDocTool/APIVariable.cs
@@ -7,6 +7,7 @@ using System.Text;
 using System.Threading.Tasks;
 using System.Xml;
 using System.IO;
+using DoxygenLib;
 
 namespace APIDocTool
 {
@@ -26,6 +27,7 @@ namespace APIDocTool
 		public string Definition = "";
 
 		public string Type = "";
+		public string AbbreviatedType = "";
 
 		public bool IsMutable = false;
 		public bool IsStatic = false;
@@ -62,6 +64,7 @@ namespace APIDocTool
 
 			XmlNode type = Node.SelectSingleNode("type");
 			Type = ConvertToMarkdown(type);
+			AbbreviatedType = Markdown.Truncate(Type, 20, "...");
 
 			XmlNodeList SimpleNodes = Node.SelectNodes("detaileddescription/para/simplesect");
 			foreach (XmlNode node in SimpleNodes)
@@ -139,8 +142,19 @@ namespace APIDocTool
 			}
         }
 
-		public UdnIconListItem GetListItem()
+		public override bool ShouldOutputPage()
 		{
+			return (MetadataDirective != null || BriefDescription != FullDescription || Type != AbbreviatedType);
+		}
+
+		public void WriteListItem(UdnWriter Writer)
+		{
+			bool OutputPage = ShouldOutputPage();
+
+			// Enter the object
+			Writer.EnterObject(OutputPage? "VariableListItem" : "VariableListItemNoLink");
+
+			// Get all the icons
 			List<Icon> Icons = new List<Icon>();
 			Icons.Add(APIDocTool.Icons.Variable[(int)Protection]);
 			if (IsStatic)
@@ -152,7 +166,42 @@ namespace APIDocTool
 				Icons.Add(APIDocTool.Icons.ReflectedVariable);
 				Icons.AddRange(MetadataDirective.Icons);
 			}
-			return new UdnIconListItem(Icons, Name, BriefDescription, LinkPath);
+			Writer.WriteParam("icons", Icons);
+
+			// Write the type
+			Writer.WriteParam("type", AbbreviatedType);
+			Writer.WriteParam("name", Name);
+			if(OutputPage)
+			{
+				Writer.WriteParam("link", "[RELATIVE:" + LinkPath + "]");
+			}
+			Writer.WriteParam("description", BriefDescription);
+
+			// Leave the object
+			Writer.LeaveObject();
+		}
+		
+		public static void WriteList(UdnWriter Writer, IEnumerable<APIVariable> Variables)
+		{
+			Writer.WriteObject("VariableListHead");
+			foreach (APIVariable Variable in Variables)
+			{
+				Variable.WriteListItem(Writer);
+			}
+			Writer.WriteObject("VariableListTail");
+		}
+
+		public static bool WriteListSection(UdnWriter Writer, string SectionId, string SectionTitle, IEnumerable<APIVariable> Variables)
+		{
+			APIVariable[] VariableArray = Variables.ToArray();
+			if (VariableArray.Length > 0)
+			{
+				Writer.EnterSection(SectionId, SectionTitle);
+				WriteList(Writer, Variables);
+				Writer.LeaveSection();
+				return true;
+			}
+			return false;
 		}
 	}
 }
diff --git a/Engine/Source/Programs/UnrealDocTool/APIDocTool/APIDocTool/Program.cs b/Engine/Source/Programs/UnrealDocTool/APIDocTool/APIDocTool/Program.cs
--- a/Engine/Source/Programs/UnrealDocTool/APIDocTool/APIDocTool/Program.cs
+++ b/Engine/Source/Programs/UnrealDocTool/APIDocTool/APIDocTool/Program.cs
@@ -256,20 +256,23 @@ namespace APIDocTool
 				Console.WriteLine();
 				Console.WriteLine("Options:");
 				Console.WriteLine("    -rebuild:                        Clean and build everything");
-				Console.WriteLine("    -rebuild<step>:					Clean and build specific steps");
+				Console.WriteLine("    -rebuild<step>:                  Clean and build specific steps");
 				Console.WriteLine("    -clean:                          Clean all files");
-				Console.WriteLine("    -clean<step>:					Clean specific steps");
-				Console.WriteLine("    -build:						    Build everything");
-				Console.WriteLine("    -build<step>:					Build specific output steps");
-				Console.WriteLine("    -archive:						Archive everything");
-				Console.WriteLine("    -archive<step>:					Archive specific output steps");
-				Console.WriteLine("    -enginedir=<...>:				Specifies the root engine directory");
-				Console.WriteLine("    -intermediatedir=<...>:			Specifies the intermediate directory");
-				Console.WriteLine("    -documentationdir=<...>:			Specifies the documentation directory");
-				Console.WriteLine("    -indexonly:						Just build index pages");
-				Console.WriteLine("    -filter=<...>,<...>:             Filter conversion, eg.");
-				Console.WriteLine("                                       Folders:  -filter=Core/Containers/...");
-				Console.WriteLine("                                       Entities: -filter=Core/TArray");
+				Console.WriteLine("    -clean<step>:                    Clean specific steps");
+				Console.WriteLine("    -build:                          Build everything");
+				Console.WriteLine("    -build<step>:                    Build specific output steps");
+				Console.WriteLine("    -archive:                        Archive everything");
+				Console.WriteLine("    -archive<step>:                  Archive specific output steps");
+				Console.WriteLine("    -enginedir=<...>:                Specifies the root engine directory");
+				Console.WriteLine("    -intermediatedir=<...>:          Specifies the intermediate directory");
+				Console.WriteLine("    -documentationdir=<...>:         Specifies the documentation directory");
+				Console.WriteLine("    -indexonly:                      Just build index pages");
+				Console.WriteLine("    -filter=<...>,<...>:             Filter which things to convert, eg.");
+				Console.WriteLine("                                     Folders:  -filter=Core/Containers/...");
+				Console.WriteLine("                                     Entities: -filter=Core/TArray");
+				Console.WriteLine("Valid steps are:");
+				Console.WriteLine("   code, codetarget, codemeta, codexml, codeudn, codehtml, codechm");
+				Console.WriteLine("   blueprint, blueprintjson, blueprintudn, blueprinthtml, blueprintchm");
 				return 1;
 			}
 
@@ -796,6 +799,10 @@ namespace APIDocTool
 				// Build a list of pages to output
 				List<APIPage> OutputPages = new List<APIPage>(Index.GatherPages().OrderBy(x => x.LinkPath));
 
+				// Remove any pages that don't want to be written
+				int NumRemoved = OutputPages.RemoveAll(x => !x.ShouldOutputPage());
+				Console.WriteLine("Removed {0} pages", NumRemoved);
+
 				// Dump the output stats
 				string StatsPath = Path.Combine(SitemapDir, "Stats.txt");
 				Console.WriteLine("Writing stats to '" + StatsPath + "'");
