Improve formatting for variables and functions, and show return types for functions.

diff --git a/Engine/Source/Programs/UnrealDocTool/APIDocTool/APIDocTool/APIFilter.cs b/Engine/Source/Programs/UnrealDocTool/APIDocTool/APIDocTool/APIFilter.cs
--- a/Engine/Source/Programs/UnrealDocTool/APIDocTool/APIDocTool/APIFilter.cs
+++ b/Engine/Source/Programs/UnrealDocTool/APIDocTool/APIDocTool/APIFilter.cs
@@ -113,7 +113,7 @@ namespace APIDocTool
 
 				// Write the function list
 				Writer.EnterTag("[PARAM:functions]");
-				APIFunction.WriteListSection(Writer, "functions", "Functions", AllFunctions);
+				APIFunction.WriteListSection(Writer, "functions", "Functions", AllFunctions, false);
 				Writer.LeaveTag("[/PARAM]");
 
 				// Write the variable list
diff --git a/Engine/Source/Programs/UnrealDocTool/APIDocTool/APIDocTool/APIFunctionGroup.cs b/Engine/Source/Programs/UnrealDocTool/APIDocTool/APIDocTool/APIFunctionGroup.cs
--- a/Engine/Source/Programs/UnrealDocTool/APIDocTool/APIDocTool/APIFunctionGroup.cs
+++ b/Engine/Source/Programs/UnrealDocTool/APIDocTool/APIDocTool/APIFunctionGroup.cs
@@ -75,7 +75,7 @@ namespace APIDocTool
 				Writer.WritePageHeader(Name, PageCrumbs, "Overload list");
 
 				Writer.EnterTag("[REGION:members]");
-				APIFunction.WriteList(Writer, Overloads);
+				APIFunction.WriteList(Writer, Overloads, Overloads.First().FunctionType != APIFunctionType.Constructor);
 				Writer.LeaveTag("[/REGION]");
 			}
 		}
diff --git a/Engine/Source/Programs/UnrealDocTool/APIDocTool/APIDocTool/APIMember.cs b/Engine/Source/Programs/UnrealDocTool/APIDocTool/APIDocTool/APIMember.cs
--- a/Engine/Source/Programs/UnrealDocTool/APIDocTool/APIDocTool/APIMember.cs
+++ b/Engine/Source/Programs/UnrealDocTool/APIDocTool/APIDocTool/APIMember.cs
@@ -192,6 +192,28 @@ namespace APIDocTool
 			return Markdown.ParseXml(Node, Indent, ResolveDoxygenLink);
 		}
 
+		public static string RemoveElaborations(string Definition)
+		{
+			string Text = Definition;
+			if (Text.StartsWith("class "))
+			{
+				Text = Text.Substring(6);
+			}
+			else if (Text.StartsWith("const class "))
+			{
+				Text = "const " + Text.Substring(12);
+			}
+			else if (Text.StartsWith("struct "))
+			{
+				Text = Text.Substring(7);
+			}
+			else if (Text.StartsWith("const struct "))
+			{
+				Text = "const " + Text.Substring(13);
+			}
+			return Text;
+		}
+
 		public static string ResolveDoxygenLink(string RefId)
 		{
 			APIMember Member = ResolveRefLink(RefId);
diff --git a/Engine/Source/Programs/UnrealDocTool/APIDocTool/APIDocTool/APIRecord.cs b/Engine/Source/Programs/UnrealDocTool/APIDocTool/APIDocTool/APIRecord.cs
--- a/Engine/Source/Programs/UnrealDocTool/APIDocTool/APIDocTool/APIRecord.cs
+++ b/Engine/Source/Programs/UnrealDocTool/APIDocTool/APIDocTool/APIRecord.cs
@@ -395,7 +395,7 @@ namespace APIDocTool
 				APIVariable.WriteListSection(Writer, "variables", "Variables", Children.OfType<APIVariable>().Where(x => x.Protection != APIProtection.Private && !x.IsDeprecated()).OrderBy(x => x.Name));
 
 				// Write all the constructors
-				if (!APIFunction.WriteListSection(Writer, "constructor", "Constructors", AllFunctions.Where(x => x.FunctionType == APIFunctionType.Constructor).OrderBy(x => x.LinkPath)) && HasAnyPrivateFunction(Name))
+				if (!APIFunction.WriteListSection(Writer, "constructor", "Constructors", AllFunctions.Where(x => x.FunctionType == APIFunctionType.Constructor).OrderBy(x => x.LinkPath), false) && HasAnyPrivateFunction(Name))
 				{
 					Writer.EnterSection("constructor", "Constructors");
 					Writer.WriteLine("No constructors are accessible with public or protected access.");
@@ -403,7 +403,7 @@ namespace APIDocTool
 				}
 
 				// Write all the destructors
-				if (!APIFunction.WriteListSection(Writer, "destructor", "Destructors", AllFunctions.Where(x => x.FunctionType == APIFunctionType.Destructor)) && HasAnyPrivateFunction("~" + Name))
+				if (!APIFunction.WriteListSection(Writer, "destructor", "Destructors", AllFunctions.Where(x => x.FunctionType == APIFunctionType.Destructor), false) && HasAnyPrivateFunction("~" + Name))
 				{
 					Writer.EnterSection("destructors", "Destructors");
 					Writer.WriteLine("No destructors are accessible with public or protected access.");
@@ -430,12 +430,12 @@ namespace APIDocTool
 					{
 						string Id = String.Format("functions_{0}", Idx);
 						string Label = (Idx == 0) ? "Functions" : String.Format("Overridden from {0}", AllBaseRecords[Idx].Name);
-						APIFunction.WriteListSection(Writer, Id, Label, AllBaseFunctions[Idx]);
+						APIFunction.WriteListSection(Writer, Id, Label, AllBaseFunctions[Idx], true);
 					}
 				}
 
 				// Write the operator list
-				APIFunction.WriteListSection(Writer, "operators", "Operators", AllFunctions.Where(x => x.FunctionType == APIFunctionType.UnaryOperator || x.FunctionType == APIFunctionType.BinaryOperator));
+				APIFunction.WriteListSection(Writer, "operators", "Operators", AllFunctions.Where(x => x.FunctionType == APIFunctionType.UnaryOperator || x.FunctionType == APIFunctionType.BinaryOperator), true);
 
 				// Write all the inner structures
 				Writer.WriteListSection("classes", "Classes", "Name", "Description", Children.OfType<APIRecord>().Where(x => x.Protection != APIProtection.Private).OrderBy(x => x.Name).Select(x => x.GetListItem()));
diff --git a/Engine/Source/Programs/UnrealDocTool/APIDocTool/APIDocTool/APIVariable.cs b/Engine/Source/Programs/UnrealDocTool/APIDocTool/APIDocTool/APIVariable.cs
--- a/Engine/Source/Programs/UnrealDocTool/APIDocTool/APIDocTool/APIVariable.cs
+++ b/Engine/Source/Programs/UnrealDocTool/APIDocTool/APIDocTool/APIVariable.cs
@@ -63,7 +63,7 @@ namespace APIDocTool
 			}
 
 			XmlNode type = Node.SelectSingleNode("type");
-			Type = ConvertToMarkdown(type);
+			Type = APIMember.RemoveElaborations(ConvertToMarkdown(type));
 			AbbreviatedType = Markdown.Truncate(Type, 20, "...");
 
 			XmlNodeList SimpleNodes = Node.SelectNodes("detaileddescription/para/simplesect");
diff --git a/Engine/Source/Programs/UnrealDocTool/DoxygenLib/Markdown.cs b/Engine/Source/Programs/UnrealDocTool/DoxygenLib/Markdown.cs
--- a/Engine/Source/Programs/UnrealDocTool/DoxygenLib/Markdown.cs
+++ b/Engine/Source/Programs/UnrealDocTool/DoxygenLib/Markdown.cs
@@ -403,7 +403,6 @@ namespace DoxygenLib
 			return Result.ToString();
 		}
 
-
 		/** 
 		 * Truncates a line to a given width, respecting embedded links.
          * 
@@ -499,6 +498,16 @@ namespace DoxygenLib
             return -1;
         }
 
+		public static int MeasureText(string Text)
+		{
+			int Length = 0;
+			for(int Idx = 0; Idx < Text.Length; Idx = SkipCharacter(Text, Idx))
+			{
+				Length++;
+			}
+			return Length;
+		}
+
 		public static int SkipCharacter(string Text, int Idx)
 		{
 			if(Text[Idx] == '&')
