Missing file for APIDocTool

diff --git a/Engine/Source/Programs/UnrealDocTool/APIDocTool/APIDocTool/APIFunction.cs b/Engine/Source/Programs/UnrealDocTool/APIDocTool/APIDocTool/APIFunction.cs
--- a/Engine/Source/Programs/UnrealDocTool/APIDocTool/APIDocTool/APIFunction.cs
+++ b/Engine/Source/Programs/UnrealDocTool/APIDocTool/APIDocTool/APIFunction.cs
@@ -224,28 +224,6 @@ namespace APIDocTool
 			get { return Node.SelectSingleNode("type").InnerText; }
 		}
 
-		public string GetAbbreviatedDefinition()
-		{
-			string Text = Definition;
-			if (Text.StartsWith("class "))
-			{
-				Text = Text.Substring(6);
-			}
-			else if (Text.StartsWith("const class "))
-			{
-				Text = "const " + Text.Substring(12);
-			}
-			else if (Text.StartsWith("struct "))
-			{
-				Text = Text.Substring(7);
-			}
-			else if (Text.StartsWith("const struct "))
-			{
-				Text = "const " + Text.Substring(13);
-			}
-			return Text;
-		}
-
 		public string Definition
 		{
 			get { return (Name == null) ? Type : (Type + " " + Name); }
@@ -448,7 +426,7 @@ namespace APIDocTool
 
 			// Get the return type
 			XmlNode ReturnTypeNode = Node.SelectSingleNode("type");
-			ReturnType = ConvertToMarkdown(ReturnTypeNode);
+			ReturnType = RemoveElaborations(ConvertToMarkdown(ReturnTypeNode));
 
 			// Parse the reimplements list
 			using (XmlNodeList ReimplementsList = Node.SelectNodes("reimplements"))
@@ -748,10 +726,10 @@ namespace APIDocTool
 			}
         }
 
-		public void WriteListItem(UdnWriter Writer)
+		public void WriteListItem(UdnWriter Writer, bool bWithType)
 		{
 			// Enter the object
-			Writer.EnterObject("FunctionListItem");
+			Writer.EnterObject(bWithType? "FunctionListItemWithType" : "FunctionListItem");
 
 			// Get all the icons
 			List<Icon> ItemIcons = new List<Icon>{ Icons.Function[(int)Protection] };
@@ -770,49 +748,58 @@ namespace APIDocTool
 			}
 			Writer.WriteParam("icons", ItemIcons);
 
-			// Write the name
-			Writer.WriteParam("name", Name);
-			Writer.WriteParam("link", "[RELATIVE:" + LinkPath + "]");
+			// Write the return type
+			Writer.WriteParam("type", Markdown.Truncate(ReturnType, 12, "..."));
 
-			// Add the parameter section if need be
-			Writer.EnterParam("arguments");
-			if (Parameters.Count > 0)
+			// Write the name
+			if(Parameters.Count == 0)
 			{
-				Writer.WriteEscapedLine("(  ");
-				for (int Idx = 0; Idx < Parameters.Count; Idx++)
+				Writer.WriteParam("name", Name + "()");
+				Writer.WriteParam("arguments", "");
+			}
+			else
+			{
+				Writer.WriteParam("name", Name);
+				Writer.EnterParam("arguments");
+				if (Parameters.Count > 0)
 				{
-					string Separator = (Idx + 1 == Parameters.Count) ? "" : ",";
-					string Definition = Markdown.Truncate(Parameters[Idx].GetAbbreviatedDefinition(), 35, "...");
-					Writer.WriteLine(UdnWriter.TabSpaces + Definition + Separator + "  ");
+					Writer.WriteEscapedLine("(  ");
+					for (int Idx = 0; Idx < Parameters.Count; Idx++)
+					{
+						string Separator = (Idx + 1 == Parameters.Count) ? "" : ",";
+						string Definition = Markdown.Truncate(APIMember.RemoveElaborations(Parameters[Idx].Definition), 35, "...");
+						Writer.WriteLine(UdnWriter.TabSpaces + Definition + Separator + "  ");
+					}
+					Writer.WriteEscapedLine(")  ");
 				}
-				Writer.WriteEscapedLine(")  ");
+				Writer.LeaveParam();
 			}
-			Writer.LeaveParam();
 
-			// Write the description
+			// Write the other parameters
+			Writer.WriteParam("link", "[RELATIVE:" + LinkPath + "]");
 			Writer.WriteParam("description", BriefDescription);
 
 			// Leave the object
 			Writer.LeaveObject();
 		}
 		
-		public static void WriteList(UdnWriter Writer, IEnumerable<APIFunction> Functions)
+		public static void WriteList(UdnWriter Writer, IEnumerable<APIFunction> Functions, bool bWithType)
 		{
-			Writer.WriteObject("FunctionListHead");
+			Writer.WriteObject(bWithType? "FunctionListHeadWithType" : "FunctionListHead");
 			foreach (APIFunction Function in Functions)
 			{
-				Function.WriteListItem(Writer);
+				Function.WriteListItem(Writer, bWithType);
 			}
 			Writer.WriteObject("FunctionListTail");
 		}
 
-		public static bool WriteListSection(UdnWriter Writer, string SectionId, string SectionTitle, IEnumerable<APIFunction> Functions)
+		public static bool WriteListSection(UdnWriter Writer, string SectionId, string SectionTitle, IEnumerable<APIFunction> Functions, bool bWithType)
 		{
 			APIFunction[] FunctionArray = Functions.ToArray();
 			if (FunctionArray.Length > 0)
 			{
 				Writer.EnterSection(SectionId, SectionTitle);
-				WriteList(Writer, Functions);
+				WriteList(Writer, Functions, bWithType);
 				Writer.LeaveSection();
 				return true;
 			}
