Remove the 'warning:' tag from Doxygen's include not found messages; they're not found because we filter them out!

diff --git a/Engine/Source/Programs/UnrealDocTool/DoxygenLib/Doxygen.cs b/Engine/Source/Programs/UnrealDocTool/DoxygenLib/Doxygen.cs
--- a/Engine/Source/Programs/UnrealDocTool/DoxygenLib/Doxygen.cs
+++ b/Engine/Source/Programs/UnrealDocTool/DoxygenLib/Doxygen.cs
@@ -170,7 +170,7 @@ namespace DoxygenLib
 					using (TextWriter SynchronizedDoxygenOutputWriter = TextWriter.Synchronized(DoxygenOutputWriter))
 					{
 						DoxygenProcess.OutputDataReceived += new DataReceivedEventHandler((x, y) => OutputReceived(x, y, SynchronizedDoxygenOutputWriter, bVerbose));
-						DoxygenProcess.ErrorDataReceived += new DataReceivedEventHandler((x, y) => ErrorReceived(x, y, SynchronizedDoxygenOutputWriter, bVerbose));
+						DoxygenProcess.ErrorDataReceived += new DataReceivedEventHandler((x, y) => OutputReceived(x, y, SynchronizedDoxygenOutputWriter, bVerbose));
 						try
 						{
 							DoxygenProcess.Start();
@@ -193,17 +193,17 @@ namespace DoxygenLib
 		{
 			if (Line.Data != null)
 			{
-				LogWriter.WriteLine(Line.Data);
-				if (bVerbose) Console.WriteLine("  doxygen> " + Line.Data);
-			}
-		}
+				// Replace the 'warning:' tag with 'pedantic:' in missing include path messages, so as not to trip the build machines.
+				// Most missing includes are by design because they're filtered out, and Doxygen is not being smart.
+				string OutputLine = Line.Data;
+				if(OutputLine.Contains("perhaps you forgot to add its directory to INCLUDE_PATH?"))
+				{
+					OutputLine = OutputLine.Replace("warning:", "pedantic:");
+				}
 
-		static private void ErrorReceived(Object Sender, DataReceivedEventArgs Line, TextWriter LogWriter, bool bVerbose)
-		{
-			if (Line.Data != null)
-			{
-				LogWriter.WriteLine(Line.Data);
-				if (bVerbose) Console.WriteLine("  doxygen> " + Line.Data);
+				// Write it to the log
+				LogWriter.WriteLine(OutputLine);
+				if (bVerbose) Console.WriteLine("  doxygen> " + OutputLine);
 			}
 		}
 	}
