Try to remove even more variable pages from the API documentation

diff --git a/Engine/Source/Programs/UnrealDocTool/APIDocTool/APIDocTool/APIVariable.cs b/Engine/Source/Programs/UnrealDocTool/APIDocTool/APIDocTool/APIVariable.cs
--- a/Engine/Source/Programs/UnrealDocTool/APIDocTool/APIDocTool/APIVariable.cs
+++ b/Engine/Source/Programs/UnrealDocTool/APIDocTool/APIDocTool/APIVariable.cs
@@ -144,7 +144,7 @@ namespace APIDocTool
 
 		public override bool ShouldOutputPage()
 		{
-			return (MetadataDirective != null || BriefDescription != FullDescription || Type != AbbreviatedType);
+			return (BriefDescription != FullDescription || Type != AbbreviatedType);
 		}
 
 		public void WriteListItem(UdnWriter Writer)
