Add source HTML files to APIDocTool project.

diff --git a/Engine/Source/Programs/UnrealDocTool/APIDocTool/APIDocTool/APIDocTool.csproj b/Engine/Source/Programs/UnrealDocTool/APIDocTool/APIDocTool/APIDocTool.csproj
--- a/Engine/Source/Programs/UnrealDocTool/APIDocTool/APIDocTool/APIDocTool.csproj
+++ b/Engine/Source/Programs/UnrealDocTool/APIDocTool/APIDocTool/APIDocTool.csproj
@@ -120,6 +120,149 @@
       <Name>DoxygenLib</Name>
     </ProjectReference>
   </ItemGroup>
+  <ItemGroup>
+    <Content Include="..\..\..\..\..\Documentation\Source\Include\CSS\udn_public.css">
+      <Link>Source\CSS\udn_public.css</Link>
+    </Content>
+    <Content Include="..\..\..\..\..\Documentation\Source\Include\Templates\Objects\ActionPinListItem.html">
+      <Link>Source\Objects\ActionPinListItem.html</Link>
+    </Content>
+    <Content Include="..\..\..\..\..\Documentation\Source\Include\Templates\Objects\BlueprintNode.html">
+      <Link>Source\Objects\BlueprintNode.html</Link>
+    </Content>
+    <Content Include="..\..\..\..\..\Documentation\Source\Include\Templates\Objects\BlueprintPin.html">
+      <Link>Source\Objects\BlueprintPin.html</Link>
+    </Content>
+    <Content Include="..\..\..\..\..\Documentation\Source\Include\Templates\Objects\Feature.html">
+      <Link>Source\Objects\Feature.html</Link>
+    </Content>
+    <Content Include="..\..\..\..\..\Documentation\Source\Include\Templates\Objects\Filter.html">
+      <Link>Source\Objects\Filter.html</Link>
+    </Content>
+    <Content Include="..\..\..\..\..\Documentation\Source\Include\Templates\Objects\FilterItem.html">
+      <Link>Source\Objects\FilterItem.html</Link>
+    </Content>
+    <Content Include="..\..\..\..\..\Documentation\Source\Include\Templates\Objects\FunctionItem.html">
+      <Link>Source\Objects\FunctionItem.html</Link>
+    </Content>
+    <Content Include="..\..\..\..\..\Documentation\Source\Include\Templates\Objects\FunctionListHead.html">
+      <Link>Source\Objects\FunctionListHead.html</Link>
+    </Content>
+    <Content Include="..\..\..\..\..\Documentation\Source\Include\Templates\Objects\FunctionListHeadWithType.html">
+      <Link>Source\Objects\FunctionListHeadWithType.html</Link>
+    </Content>
+    <Content Include="..\..\..\..\..\Documentation\Source\Include\Templates\Objects\FunctionListItem.html">
+      <Link>Source\Objects\FunctionListItem.html</Link>
+    </Content>
+    <Content Include="..\..\..\..\..\Documentation\Source\Include\Templates\Objects\FunctionListItemWithType.html">
+      <Link>Source\Objects\FunctionListItemWithType.html</Link>
+    </Content>
+    <Content Include="..\..\..\..\..\Documentation\Source\Include\Templates\Objects\FunctionListTail.html">
+      <Link>Source\Objects\FunctionListTail.html</Link>
+    </Content>
+    <Content Include="..\..\..\..\..\Documentation\Source\Include\Templates\Objects\HierarchyButton.html">
+      <Link>Source\Objects\HierarchyButton.html</Link>
+    </Content>
+    <Content Include="..\..\..\..\..\Documentation\Source\Include\Templates\Objects\HierarchyLabel.html">
+      <Link>Source\Objects\HierarchyLabel.html</Link>
+    </Content>
+    <Content Include="..\..\..\..\..\Documentation\Source\Include\Templates\Objects\HierarchyLabelLinked.html">
+      <Link>Source\Objects\HierarchyLabelLinked.html</Link>
+    </Content>
+    <Content Include="..\..\..\..\..\Documentation\Source\Include\Templates\Objects\HierarchySpacer.html">
+      <Link>Source\Objects\HierarchySpacer.html</Link>
+    </Content>
+    <Content Include="..\..\..\..\..\Documentation\Source\Include\Templates\Objects\MemberHeader.html">
+      <Link>Source\Objects\MemberHeader.html</Link>
+    </Content>
+    <Content Include="..\..\..\..\..\Documentation\Source\Include\Templates\Objects\MemberIconListHead.html">
+      <Link>Source\Objects\MemberIconListHead.html</Link>
+    </Content>
+    <Content Include="..\..\..\..\..\Documentation\Source\Include\Templates\Objects\MemberIconListHeadBlank.html">
+      <Link>Source\Objects\MemberIconListHeadBlank.html</Link>
+    </Content>
+    <Content Include="..\..\..\..\..\Documentation\Source\Include\Templates\Objects\MemberIconListItem.html">
+      <Link>Source\Objects\MemberIconListItem.html</Link>
+    </Content>
+    <Content Include="..\..\..\..\..\Documentation\Source\Include\Templates\Objects\MemberIconListItemLinked.html">
+      <Link>Source\Objects\MemberIconListItemLinked.html</Link>
+    </Content>
+    <Content Include="..\..\..\..\..\Documentation\Source\Include\Templates\Objects\MemberIconListTail.html">
+      <Link>Source\Objects\MemberIconListTail.html</Link>
+    </Content>
+    <Content Include="..\..\..\..\..\Documentation\Source\Include\Templates\Objects\MemberItem.html">
+      <Link>Source\Objects\MemberItem.html</Link>
+    </Content>
+    <Content Include="..\..\..\..\..\Documentation\Source\Include\Templates\Objects\MemberList.html">
+      <Link>Source\Objects\MemberList.html</Link>
+    </Content>
+    <Content Include="..\..\..\..\..\Documentation\Source\Include\Templates\Objects\MemberListHead.html">
+      <Link>Source\Objects\MemberListHead.html</Link>
+    </Content>
+    <Content Include="..\..\..\..\..\Documentation\Source\Include\Templates\Objects\MemberListHeadBlank.html">
+      <Link>Source\Objects\MemberListHeadBlank.html</Link>
+    </Content>
+    <Content Include="..\..\..\..\..\Documentation\Source\Include\Templates\Objects\MemberListItem.html">
+      <Link>Source\Objects\MemberListItem.html</Link>
+    </Content>
+    <Content Include="..\..\..\..\..\Documentation\Source\Include\Templates\Objects\MemberListItemLinked.html">
+      <Link>Source\Objects\MemberListItemLinked.html</Link>
+    </Content>
+    <Content Include="..\..\..\..\..\Documentation\Source\Include\Templates\Objects\MemberListTail.html">
+      <Link>Source\Objects\MemberListTail.html</Link>
+    </Content>
+    <Content Include="..\..\..\..\..\Documentation\Source\Include\Templates\Objects\MetadataListHead.html">
+      <Link>Source\Objects\MetadataListHead.html</Link>
+    </Content>
+    <Content Include="..\..\..\..\..\Documentation\Source\Include\Templates\Objects\MetadataListItem.html">
+      <Link>Source\Objects\MetadataListItem.html</Link>
+    </Content>
+    <Content Include="..\..\..\..\..\Documentation\Source\Include\Templates\Objects\MetadataListTail.html">
+      <Link>Source\Objects\MetadataListTail.html</Link>
+    </Content>
+    <Content Include="..\..\..\..\..\Documentation\Source\Include\Templates\Objects\Module.html">
+      <Link>Source\Objects\Module.html</Link>
+    </Content>
+    <Content Include="..\..\..\..\..\Documentation\Source\Include\Templates\Objects\ModuleSection.html">
+      <Link>Source\Objects\ModuleSection.html</Link>
+    </Content>
+    <Content Include="..\..\..\..\..\Documentation\Source\Include\Templates\Objects\ModuleSectionExpanded.html">
+      <Link>Source\Objects\ModuleSectionExpanded.html</Link>
+    </Content>
+    <Content Include="..\..\..\..\..\Documentation\Source\Include\Templates\Objects\Navigation.html">
+      <Link>Source\Objects\Navigation.html</Link>
+    </Content>
+    <Content Include="..\..\..\..\..\Documentation\Source\Include\Templates\Objects\Param.html">
+      <Link>Source\Objects\Param.html</Link>
+    </Content>
+    <Content Include="..\..\..\..\..\Documentation\Source\Include\Templates\Objects\ParamHeader.html">
+      <Link>Source\Objects\ParamHeader.html</Link>
+    </Content>
+    <Content Include="..\..\..\..\..\Documentation\Source\Include\Templates\Objects\Typedef.html">
+      <Link>Source\Objects\Typedef.html</Link>
+    </Content>
+    <Content Include="..\..\..\..\..\Documentation\Source\Include\Templates\Objects\TypedefItem.html">
+      <Link>Source\Objects\TypedefItem.html</Link>
+    </Content>
+    <Content Include="..\..\..\..\..\Documentation\Source\Include\Templates\Objects\ValueItem.html">
+      <Link>Source\Objects\ValueItem.html</Link>
+    </Content>
+    <Content Include="..\..\..\..\..\Documentation\Source\Include\Templates\Objects\VariableItem.html">
+      <Link>Source\Objects\VariableItem.html</Link>
+    </Content>
+    <Content Include="..\..\..\..\..\Documentation\Source\Include\Templates\Objects\VariableListHead.html">
+      <Link>Source\Objects\VariableListHead.html</Link>
+    </Content>
+    <Content Include="..\..\..\..\..\Documentation\Source\Include\Templates\Objects\VariableListItem.html">
+      <Link>Source\Objects\VariableListItem.html</Link>
+    </Content>
+    <Content Include="..\..\..\..\..\Documentation\Source\Include\Templates\Objects\VariableListItemNoLink.html">
+      <Link>Source\Objects\VariableListItemNoLink.html</Link>
+    </Content>
+    <Content Include="..\..\..\..\..\Documentation\Source\Include\Templates\Objects\VariableListTail.html">
+      <Link>Source\Objects\VariableListTail.html</Link>
+    </Content>
+  </ItemGroup>
   <ItemGroup />
   <Import Project="$(MSBuildToolsPath)\Microsoft.CSharp.targets" />
   <!-- To modify your build process, add your task inside one of the targets below and uncomment it. 
