Remove deleted pages from the CHM sitemap.

diff --git a/Engine/Source/Programs/UnrealDocTool/APIDocTool/APIDocTool/APICategory.cs b/Engine/Source/Programs/UnrealDocTool/APIDocTool/APIDocTool/APICategory.cs
--- a/Engine/Source/Programs/UnrealDocTool/APIDocTool/APIDocTool/APICategory.cs
+++ b/Engine/Source/Programs/UnrealDocTool/APIDocTool/APIDocTool/APICategory.cs
@@ -47,11 +47,11 @@ namespace APIDocTool
 			Actions.Add(Action);
 		}
 
-		public override SitemapNode CreateSitemapNode()
+		public override IEnumerable<SitemapNode> CreateSitemapNodes()
 		{
 			SitemapNode Node = new SitemapNode(Name, SitemapLinkPath);
-			Node.Children.AddRange(SubCategories.OrderBy(x => x.Value.Name).Select(x => x.Value.CreateSitemapNode()));
-			return Node;
+			Node.Children.AddRange(SubCategories.OrderBy(x => x.Value.Name).SelectMany(x => x.Value.CreateSitemapNodes()));
+			yield return Node;
 		}
 
 		public void WriteSitemapContents(string OutputPath)
@@ -63,11 +63,11 @@ namespace APIDocTool
 			}
 			foreach (APIPage SubCategory in SubCategories.Values)
 			{
-				RootNodes.Add(SubCategory.CreateSitemapNode());
+				RootNodes.AddRange(SubCategory.CreateSitemapNodes());
 			}
 			foreach (APIPage Action in Actions)
 			{
-				RootNodes.Add(Action.CreateSitemapNode());
+				RootNodes.AddRange(Action.CreateSitemapNodes());
 			}
 			Sitemap.WriteContentsFile(OutputPath, RootNodes);
 		}
diff --git a/Engine/Source/Programs/UnrealDocTool/APIDocTool/APIDocTool/APIFilter.cs b/Engine/Source/Programs/UnrealDocTool/APIDocTool/APIDocTool/APIFilter.cs
--- a/Engine/Source/Programs/UnrealDocTool/APIDocTool/APIDocTool/APIFilter.cs
+++ b/Engine/Source/Programs/UnrealDocTool/APIDocTool/APIDocTool/APIFilter.cs
@@ -26,13 +26,13 @@ namespace APIDocTool
 			Pages.AddRange(Members);
 		}
 
-		public override SitemapNode CreateSitemapNode()
+		public override IEnumerable<SitemapNode> CreateSitemapNodes()
 		{
 			SitemapNode Node = new SitemapNode(Name, SitemapLinkPath);
-			Node.Children.AddRange(Filters.Select(x => x.CreateSitemapNode()));
-			Node.Children.AddRange(Members.OfType<APIRecord>().Select(x => x.CreateSitemapNode()));
+			Node.Children.AddRange(Filters.SelectMany(x => x.CreateSitemapNodes()));
+			Node.Children.AddRange(Members.OfType<APIRecord>().SelectMany(x => x.CreateSitemapNodes()));
 			if (Node.Children.Count == 0) Node.Children.Add(new SitemapNode(Name, SitemapLinkPath));
-			return Node;
+			yield return Node;
 		}
 
 		public override void AddToManifest(UdnManifest Manifest)
diff --git a/Engine/Source/Programs/UnrealDocTool/APIDocTool/APIDocTool/APIIndex.cs b/Engine/Source/Programs/UnrealDocTool/APIDocTool/APIDocTool/APIIndex.cs
--- a/Engine/Source/Programs/UnrealDocTool/APIDocTool/APIDocTool/APIIndex.cs
+++ b/Engine/Source/Programs/UnrealDocTool/APIDocTool/APIDocTool/APIIndex.cs
@@ -138,7 +138,7 @@ namespace APIDocTool
 			RootNodes.Add(new SitemapNode("Getting started", GettingStarted.SitemapLinkPath));
 			foreach (APIModuleIndex ChildModuleIndex in ChildModuleIndexes)
 			{
-				RootNodes.Add(ChildModuleIndex.CreateSitemapNode());
+				RootNodes.AddRange(ChildModuleIndex.CreateSitemapNodes());
 			}
 			RootNodes.Add(new SitemapNode("All constants", ConstantIndex.SitemapLinkPath));
 			RootNodes.Add(new SitemapNode("All functions", FunctionIndex.SitemapLinkPath));
diff --git a/Engine/Source/Programs/UnrealDocTool/APIDocTool/APIDocTool/APIModuleIndex.cs b/Engine/Source/Programs/UnrealDocTool/APIDocTool/APIDocTool/APIModuleIndex.cs
--- a/Engine/Source/Programs/UnrealDocTool/APIDocTool/APIDocTool/APIModuleIndex.cs
+++ b/Engine/Source/Programs/UnrealDocTool/APIDocTool/APIDocTool/APIModuleIndex.cs
@@ -157,11 +157,11 @@ namespace APIDocTool
 			Pages.AddRange(Children);
 		}
 
-		public override SitemapNode CreateSitemapNode()
+		public override IEnumerable<SitemapNode> CreateSitemapNodes()
 		{
 			SitemapNode Node = new SitemapNode(Name, SitemapLinkPath);
-			Node.Children.AddRange(Children.OrderBy(x => x.Name).Select(x => x.CreateSitemapNode()));
-			return Node;
+			Node.Children.AddRange(Children.OrderBy(x => x.Name).SelectMany(x => x.CreateSitemapNodes()));
+			yield return Node;
 		}
 
 		public override void AddToManifest(UdnManifest Manifest)
diff --git a/Engine/Source/Programs/UnrealDocTool/APIDocTool/APIDocTool/APIPage.cs b/Engine/Source/Programs/UnrealDocTool/APIDocTool/APIDocTool/APIPage.cs
--- a/Engine/Source/Programs/UnrealDocTool/APIDocTool/APIDocTool/APIPage.cs
+++ b/Engine/Source/Programs/UnrealDocTool/APIDocTool/APIDocTool/APIPage.cs
@@ -91,9 +91,12 @@ namespace APIDocTool
 			return true;
 		}
 
-		public virtual SitemapNode CreateSitemapNode()
+		public virtual IEnumerable<SitemapNode> CreateSitemapNodes()
 		{
-			return new SitemapNode(Name, SitemapLinkPath);
+			if(ShouldOutputPage())
+			{
+				yield return new SitemapNode(Name, SitemapLinkPath);
+			}
 		}
 
 		public abstract void WritePage(UdnManifest Manifest, string OutputPath);
