Show function types in filters.

diff --git a/Engine/Source/Programs/UnrealDocTool/APIDocTool/APIDocTool/APIFilter.cs b/Engine/Source/Programs/UnrealDocTool/APIDocTool/APIDocTool/APIFilter.cs
--- a/Engine/Source/Programs/UnrealDocTool/APIDocTool/APIDocTool/APIFilter.cs
+++ b/Engine/Source/Programs/UnrealDocTool/APIDocTool/APIDocTool/APIFilter.cs
@@ -113,7 +113,7 @@ namespace APIDocTool
 
 				// Write the function list
 				Writer.EnterTag("[PARAM:functions]");
-				APIFunction.WriteListSection(Writer, "functions", "Functions", AllFunctions, false);
+				APIFunction.WriteListSection(Writer, "functions", "Functions", AllFunctions, true);
 				Writer.LeaveTag("[/PARAM]");
 
 				// Write the variable list
