Fix some formatting in API docs.

diff --git a/Engine/Source/Programs/UnrealDocTool/APIDocTool/APIDocTool/APIDocTool.csproj b/Engine/Source/Programs/UnrealDocTool/APIDocTool/APIDocTool/APIDocTool.csproj
--- a/Engine/Source/Programs/UnrealDocTool/APIDocTool/APIDocTool/APIDocTool.csproj
+++ b/Engine/Source/Programs/UnrealDocTool/APIDocTool/APIDocTool/APIDocTool.csproj
@@ -82,13 +82,13 @@
     <Compile Include="APIRecord.cs" />
     <Compile Include="APIRecordIndex.cs" />
     <Compile Include="BuildTarget.cs" />
-    <Compile Include="fastJSON\Formatter.cs" />
-    <Compile Include="fastJSON\Getters.cs" />
-    <Compile Include="fastJSON\JSON.cs" />
-    <Compile Include="fastJSON\JsonParser.cs" />
-    <Compile Include="fastJSON\JsonSerializer.cs" />
-    <Compile Include="fastJSON\Reflection.cs" />
-    <Compile Include="fastJSON\SafeDictionary.cs" />
+    <Compile Include="FastJSON\Formatter.cs" />
+    <Compile Include="FastJSON\Getters.cs" />
+    <Compile Include="FastJSON\JSON.cs" />
+    <Compile Include="FastJSON\JsonParser.cs" />
+    <Compile Include="FastJSON\JsonSerializer.cs" />
+    <Compile Include="FastJSON\Reflection.cs" />
+    <Compile Include="FastJSON\SafeDictionary.cs" />
     <Compile Include="Icons.cs" />
     <Compile Include="IniFile.cs" />
     <Compile Include="Metadata.cs" />
@@ -122,7 +122,16 @@
   </ItemGroup>
   <ItemGroup>
     <Content Include="..\..\..\..\..\Documentation\Source\Include\CSS\udn_public.css">
-      <Link>Source\CSS\udn_public.css</Link>
+      <Link>Pages\udn_public.css</Link>
+    </Content>
+    <Content Include="..\..\..\..\..\Documentation\Source\Include\Javascript\udn.js">
+      <Link>Pages\udn.js</Link>
+    </Content>
+    <Content Include="..\..\..\..\..\Documentation\Source\Include\Templates\Default.html">
+      <Link>Pages\Default.html</Link>
+    </Content>
+    <Content Include="..\..\..\..\..\Documentation\Source\Include\Templates\DefaultAPI.html">
+      <Link>Pages\DefaultAPI.html</Link>
     </Content>
     <Content Include="..\..\..\..\..\Documentation\Source\Include\Templates\Objects\ActionPinListItem.html">
       <Link>Source\Objects\ActionPinListItem.html</Link>
diff --git a/Engine/Source/Programs/UnrealDocTool/APIDocTool/APIDocTool/APIGettingStarted.cs b/Engine/Source/Programs/UnrealDocTool/APIDocTool/APIDocTool/APIGettingStarted.cs
--- a/Engine/Source/Programs/UnrealDocTool/APIDocTool/APIDocTool/APIGettingStarted.cs
+++ b/Engine/Source/Programs/UnrealDocTool/APIDocTool/APIDocTool/APIGettingStarted.cs
@@ -38,7 +38,6 @@ namespace APIDocTool
 				Writer.WriteLine("Code in each module can use other modules by referencing them in their *build rules*. Build rules for each module are given by C# scripts with the .build.cs extension. *Target rules* are given by C# scripts with the .target.cs extension.");
 				Writer.WriteLine();
 				Writer.WriteLine(Manifest, "Modules supplied with the Unreal Engine are divided into three categories; the {Runtime|ModuleIndex:Runtime Modules}, functionality for the {Editor|ModuleIndex:Editor Modules}, and {Developer|ModuleIndex:Developer Modules} utilities.");
-				Writer.WriteLine();
 				Writer.WriteLine(Manifest, "Most gameplay programming just uses runtime modules, and the three most commonly encountered are {Core|Filter:Core}, {CoreUObject|Filter:CoreUObject} and {Engine|Filter:Engine}.");
 				Writer.WriteLine();
 
diff --git a/Engine/Source/Programs/UnrealDocTool/APIDocTool/APIDocTool/APIVariable.cs b/Engine/Source/Programs/UnrealDocTool/APIDocTool/APIDocTool/APIVariable.cs
--- a/Engine/Source/Programs/UnrealDocTool/APIDocTool/APIDocTool/APIVariable.cs
+++ b/Engine/Source/Programs/UnrealDocTool/APIDocTool/APIDocTool/APIVariable.cs
@@ -54,7 +54,13 @@ namespace APIDocTool
 
 			ParseBriefAndFullDescription(Node, out BriefDescription, out FullDescription);
 
-			GetModifiers();
+            XmlNode BitfieldNode = Node.SelectSingleNode("bitfield");
+            if (BitfieldNode != null)
+            {
+                Bitfield = BitfieldNode.InnerText;
+            }
+			IsMutable = Node.Attributes.GetNamedItem("mutable").InnerText == "yes";
+			IsStatic = Node.Attributes.GetNamedItem("static").InnerText == "yes";
 
 			XmlNode DefNode = Node.SelectSingleNode("definition");
 			if (DefNode != null)
@@ -64,7 +70,7 @@ namespace APIDocTool
 
 			XmlNode type = Node.SelectSingleNode("type");
 			Type = APIMember.RemoveElaborations(ConvertToMarkdown(type));
-			AbbreviatedType = Markdown.Truncate(Type, 20, "...");
+			AbbreviatedType = Markdown.Truncate(Type + Bitfield, 15, "...");
 
 			XmlNodeList SimpleNodes = Node.SelectNodes("detaileddescription/para/simplesect");
 			foreach (XmlNode node in SimpleNodes)
@@ -78,17 +84,6 @@ namespace APIDocTool
 			}
 		}
 
-        public void GetModifiers()
-        {
-            XmlNode BitfieldNode = Node.SelectSingleNode("bitfield");
-            if (BitfieldNode != null)
-            {
-                Bitfield = BitfieldNode.InnerText;
-            }
-			IsMutable = Node.Attributes.GetNamedItem("mutable").InnerText == "yes";
-			IsStatic = Node.Attributes.GetNamedItem("static").InnerText == "yes";
-        }
-
 		private void WriteDefinition(UdnWriter Writer)
 		{
 			List<string> Lines = new List<string>();
