Fix formatting for bitfields in API docs.

diff --git a/Engine/Source/Programs/UnrealDocTool/APIDocTool/APIDocTool/APIVariable.cs b/Engine/Source/Programs/UnrealDocTool/APIDocTool/APIDocTool/APIVariable.cs
--- a/Engine/Source/Programs/UnrealDocTool/APIDocTool/APIDocTool/APIVariable.cs
+++ b/Engine/Source/Programs/UnrealDocTool/APIDocTool/APIDocTool/APIVariable.cs
@@ -27,6 +27,7 @@ namespace APIDocTool
 		public string Definition = "";
 
 		public string Type = "";
+		public string IsolatedType = "";
 		public string AbbreviatedType = "";
 
 		public bool IsMutable = false;
@@ -70,7 +71,13 @@ namespace APIDocTool
 
 			XmlNode type = Node.SelectSingleNode("type");
 			Type = APIMember.RemoveElaborations(ConvertToMarkdown(type));
-			AbbreviatedType = Markdown.Truncate(Type + Bitfield, 15, "...");
+
+			IsolatedType = Type;
+			if(!String.IsNullOrEmpty(Bitfield))
+			{
+				IsolatedType += ": " + Bitfield;
+			}
+			AbbreviatedType = Markdown.Truncate(IsolatedType, 15, "...");
 
 			XmlNodeList SimpleNodes = Node.SelectNodes("detaileddescription/para/simplesect");
 			foreach (XmlNode node in SimpleNodes)
@@ -139,7 +146,7 @@ namespace APIDocTool
 
 		public override bool ShouldOutputPage()
 		{
-			return (BriefDescription != FullDescription || Type != AbbreviatedType);
+			return (BriefDescription != FullDescription || IsolatedType != AbbreviatedType);
 		}
 
 		public void WriteListItem(UdnWriter Writer)
