Various APIDocTool fixes:

* Manifest references first parent page that is output, rather than referencing pages that don't exist
* CHM index no longer includes variable pages that have been removed
* Pointer symbols being translated as italic markdown in function prototypes
* Missing array decorators from all variable declarations
* Fixed formatting of top-level class hierarchy links

diff --git a/Engine/Source/Programs/UnrealDocTool/APIDocTool/APIDocTool/APIIndex.cs b/Engine/Source/Programs/UnrealDocTool/APIDocTool/APIDocTool/APIIndex.cs
--- a/Engine/Source/Programs/UnrealDocTool/APIDocTool/APIDocTool/APIIndex.cs
+++ b/Engine/Source/Programs/UnrealDocTool/APIDocTool/APIDocTool/APIIndex.cs
@@ -174,7 +174,7 @@ namespace APIDocTool
 			List<SitemapNode> RootNodes = new List<SitemapNode>();
 			foreach(APIPage Page in AllPages)
 			{
-				if (!(Page is APIMember))
+				if (!(Page is APIMember) && Page.ShouldOutputPage())
 				{
 					RootNodes.Add(new SitemapNode(Page.Name, Page.SitemapLinkPath));
 				}
@@ -183,7 +183,10 @@ namespace APIDocTool
 			{
 				foreach(APIMember Member in MemberIndexLink.Value)
 				{
-					RootNodes.Add(new SitemapNode(MemberIndexLink.Key, Member.SitemapLinkPath));
+					if(Member.ShouldOutputPage())
+					{
+						RootNodes.Add(new SitemapNode(MemberIndexLink.Key, Member.SitemapLinkPath));
+					}
 				}
 			}
 
diff --git a/Engine/Source/Programs/UnrealDocTool/APIDocTool/APIDocTool/APIVariable.cs b/Engine/Source/Programs/UnrealDocTool/APIDocTool/APIDocTool/APIVariable.cs
--- a/Engine/Source/Programs/UnrealDocTool/APIDocTool/APIDocTool/APIVariable.cs
+++ b/Engine/Source/Programs/UnrealDocTool/APIDocTool/APIDocTool/APIVariable.cs
@@ -16,8 +16,9 @@ namespace APIDocTool
 		public readonly XmlNode Node;
 		public readonly APIProtection Protection;
 
-        public Boolean Mutable;
-        public String Bitfield = "";
+        public bool Mutable;
+        public string Bitfield = "";
+		public string ArgsString = "";
 
 		public string BriefDescription = "";
 		public string FullDescription = "";
@@ -72,7 +73,17 @@ namespace APIDocTool
 			XmlNode type = Node.SelectSingleNode("type");
 			Type = APIMember.RemoveElaborations(ConvertToMarkdown(type));
 
+			XmlNode ArgsStringNode = Node.SelectSingleNode("argsstring");
+			if(ArgsStringNode != null)
+			{
+				ArgsString = ConvertToMarkdown(ArgsStringNode);
+			}
+
 			IsolatedType = Type;
+			if(!String.IsNullOrEmpty(ArgsString))
+			{
+				IsolatedType += ArgsString;
+			}
 			if(!String.IsNullOrEmpty(Bitfield))
 			{
 				IsolatedType += ": " + Bitfield;
@@ -105,6 +116,7 @@ namespace APIDocTool
 			if (IsMutable) Definition.Append("mutable ");
 			Definition.Append(Type + " " + Name);
 			if (Bitfield != "") Definition.Append(": " + Bitfield);
+			if (ArgsString != "") Definition.Append(ArgsString);
 			Definition.Append("  ");
 			Lines.Add(Definition.ToString());
 
diff --git a/Engine/Source/Programs/UnrealDocTool/APIDocTool/APIDocTool/UdnManifest.cs b/Engine/Source/Programs/UnrealDocTool/APIDocTool/APIDocTool/UdnManifest.cs
--- a/Engine/Source/Programs/UnrealDocTool/APIDocTool/APIDocTool/UdnManifest.cs
+++ b/Engine/Source/Programs/UnrealDocTool/APIDocTool/APIDocTool/UdnManifest.cs
@@ -122,7 +122,12 @@ namespace APIDocTool
 			{
 				foreach (KeyValuePair<string, APIPage> Entry in Entries.OrderBy(x => x.Key))
 				{
-					Writer.WriteLine("{0}, {1}", Entry.Key, Entry.Value.LinkPath);
+					APIPage TargetPage = Entry.Value;
+					while(!TargetPage.ShouldOutputPage())
+					{
+						TargetPage = TargetPage.Parent;
+					}
+					Writer.WriteLine("{0}, {1}", Entry.Key, TargetPage.LinkPath);
 				}
 			}
 		}
diff --git a/Engine/Source/Programs/UnrealDocTool/DoxygenLib/Markdown.cs b/Engine/Source/Programs/UnrealDocTool/DoxygenLib/Markdown.cs
--- a/Engine/Source/Programs/UnrealDocTool/DoxygenLib/Markdown.cs
+++ b/Engine/Source/Programs/UnrealDocTool/DoxygenLib/Markdown.cs
@@ -395,6 +395,10 @@ namespace DoxygenLib
 				{
 					Result.Append("&#33;");
 				}
+				else if(Text[Idx] == '*')
+				{
+					Result.Append("&#42;");
+				}
 				else
 				{
 					Result.Append(Text[Idx]);
