Integrate CHM change back from 4.6 branch

diff --git a/Engine/Source/Programs/UnrealDocTool/APIDocTool/APIDocTool/Program.cs b/Engine/Source/Programs/UnrealDocTool/APIDocTool/APIDocTool/Program.cs
--- a/Engine/Source/Programs/UnrealDocTool/APIDocTool/APIDocTool/Program.cs
+++ b/Engine/Source/Programs/UnrealDocTool/APIDocTool/APIDocTool/Program.cs
@@ -1200,6 +1200,10 @@ namespace APIDocTool
 				Writer.WriteLine("Full-text search=Yes");
 				Writer.WriteLine("Display compile progress=Yes");
 				Writer.WriteLine("Language=0x409 English (United States)");
+				Writer.WriteLine("Default Window=MainWindow");
+				Writer.WriteLine();
+				Writer.WriteLine("[WINDOWS]");
+				Writer.WriteLine("MainWindow=\"{0}\",\"{1}.hhc\",\"{1}.hhk\",\"{2}\",\"{2}\",,,,,0x23520,230,0x304e,[10,10,1260,750],,,,,,,0", Title, ProjectName, DefaultTopicPath);
 				Writer.WriteLine();
 				Writer.WriteLine("[FILES]");
 				foreach (string FileName in FileNames)
