Few macro definition to help Doxygen parse the latest code.

diff --git a/Engine/Source/Programs/UnrealDocTool/APIDocTool/APIDocTool/Program.cs b/Engine/Source/Programs/UnrealDocTool/APIDocTool/APIDocTool/Program.cs
--- a/Engine/Source/Programs/UnrealDocTool/APIDocTool/APIDocTool/Program.cs
+++ b/Engine/Source/Programs/UnrealDocTool/APIDocTool/APIDocTool/Program.cs
@@ -117,6 +117,7 @@ namespace APIDocTool
 			"WITH_EDITORONLY_DATA=1",
 			"WITH_PHYSX=1",
 			"WITH_SUBSTEPPING=1",
+			"WITH_HOT_RELOAD=1",
 			"NAVOCTREE_CONTAINS_COLLISION_DATA=1",
 			"SOURCE_CONTROL_WITH_SLATE=1",
 			"MATCHMAKING_HACK_FOR_EGP_IE_HOSTING=1",
@@ -166,6 +167,10 @@ namespace APIDocTool
 
 			// Log categories
 			"DECLARE_LOG_CATEGORY_EXTERN(CategoryName, DefaultVerbosity, CompileTimeVerbosity)= ",
+
+			// Delegates
+			"FUNC_PARAM_MEMBERS=",
+			"DECLARE_DERIVED_EVENT(X,Y,Z)="
 		};
 
 		static Program()
