More Doxygen warning fixes.

diff --git a/Engine/Source/Developer/NewLogVisualizer/Private/TimeSliderController.cpp b/Engine/Source/Developer/NewLogVisualizer/Private/TimeSliderController.cpp
--- a/Engine/Source/Developer/NewLogVisualizer/Private/TimeSliderController.cpp
+++ b/Engine/Source/Developer/NewLogVisualizer/Private/TimeSliderController.cpp
@@ -127,7 +127,7 @@ struct FDrawTickArgs
 	
 };
 
-void FSequencerTimeSliderController::DrawTicks( FSlateWindowElementList& OutDrawElements, const FScrubRangeToScreen& RangeToScreen, FDrawTickArgs& InArgs ) const
+void FSequencerTimeSliderController::DrawTicks( FSlateWindowElementList& OutDrawElements, const struct FScrubRangeToScreen& RangeToScreen, FDrawTickArgs& InArgs ) const
 {
 	const float Spacing = DetermineOptimalSpacing( RangeToScreen.PixelsPerInput, ScrubConstants::MinPixelsPerDisplayTick, ScrubConstants::MinDisplayTickSpacing );
 
diff --git a/Engine/Source/Editor/GraphEditor/Private/DragConnection.h b/Engine/Source/Editor/GraphEditor/Private/DragConnection.h
--- a/Engine/Source/Editor/GraphEditor/Private/DragConnection.h
+++ b/Engine/Source/Editor/GraphEditor/Private/DragConnection.h
@@ -2,6 +2,8 @@
 
 #pragma once
 
+#include "DragAndDrop.h"
+
 class FDragConnection : public FGraphEditorDragDropAction
 {
 public:
diff --git a/Engine/Source/Editor/GraphEditor/Public/GraphEditorDragDropAction.h b/Engine/Source/Editor/GraphEditor/Public/GraphEditorDragDropAction.h
--- a/Engine/Source/Editor/GraphEditor/Public/GraphEditorDragDropAction.h
+++ b/Engine/Source/Editor/GraphEditor/Public/GraphEditorDragDropAction.h
@@ -2,6 +2,8 @@
 
 #pragma once
 
+#include "DragAndDrop.h"
+
 // Base class for drag-drop actions that pass into the graph editor and perform an action when dropped
 class GRAPHEDITOR_API FGraphEditorDragDropAction : public FDragDropOperation
 {
diff --git a/Engine/Source/Programs/UnrealDocTool/APIDocTool/APIDocTool/Program.cs b/Engine/Source/Programs/UnrealDocTool/APIDocTool/APIDocTool/Program.cs
--- a/Engine/Source/Programs/UnrealDocTool/APIDocTool/APIDocTool/Program.cs
+++ b/Engine/Source/Programs/UnrealDocTool/APIDocTool/APIDocTool/Program.cs
@@ -136,6 +136,7 @@ namespace APIDocTool
 
 			// Slate declarative syntax
 			"SLATE_BEGIN_ARGS(WidgetType)=public: struct FArguments : public TSlateBaseNamedArgs<WidgetType> { typedef FArguments WidgetArgsType; FArguments()",
+			"SLATE_USER_ARGS(WidgetType)=public: static TSharedRef<WidgetType> New(); struct FArguments; struct FArguments : public TSlateBaseNamedArgs<WidgetType>{ typedef FArguments WidgetArgsType; FArguments()",
 			"HACK_SLATE_SLOT_ARGS(WidgetType)=public: struct FArguments : public TSlateBaseNamedArgs<WidgetType>{ typedef FArguments WidgetArgsType; FArguments()",
 			"SLATE_ATTRIBUTE(AttrType, AttrName)=WidgetArgsType &AttrName( const TAttribute<AttrType>& InAttribute );",
 			"SLATE_TEXT_ATTRIBUTE(AttrName)=WidgetArgsType &AttrName( const TAttribute<FText>& InAttribute ); WidgetArgsType &AttrName( const TAttribute<FString>& InAttribute );",
diff --git a/Engine/Source/Runtime/Engine/Public/Model.h b/Engine/Source/Runtime/Engine/Public/Model.h
--- a/Engine/Source/Runtime/Engine/Public/Model.h
+++ b/Engine/Source/Runtime/Engine/Public/Model.h
@@ -6,6 +6,7 @@
 
 #pragma once
 
+#include "ObjectBase.h"
 #include "LocalVertexFactory.h"
 #include "SceneTypes.h"
 #include "RawIndexBuffer.h"
diff --git a/Engine/Source/Runtime/Slate/Public/Framework/Commands/UICommandDragDropOp.h b/Engine/Source/Runtime/Slate/Public/Framework/Commands/UICommandDragDropOp.h
--- a/Engine/Source/Runtime/Slate/Public/Framework/Commands/UICommandDragDropOp.h
+++ b/Engine/Source/Runtime/Slate/Public/Framework/Commands/UICommandDragDropOp.h
@@ -2,6 +2,7 @@
 
 #pragma once
 
+#include "DragAndDrop.h"
 
 /**
  * A drag drop operation for UI Commands
