Mask warnings from HHC about duplicated pages when building API docs. The warnings don't seem to be right; possibly hash collisions?

diff --git a/Engine/Source/Programs/UnrealDocTool/APIDocTool/APIDocTool/Program.cs b/Engine/Source/Programs/UnrealDocTool/APIDocTool/APIDocTool/Program.cs
--- a/Engine/Source/Programs/UnrealDocTool/APIDocTool/APIDocTool/Program.cs
+++ b/Engine/Source/Programs/UnrealDocTool/APIDocTool/APIDocTool/Program.cs
@@ -1232,8 +1232,8 @@ namespace APIDocTool
 				CompilerProcess.StartInfo.UseShellExecute = false;
 				CompilerProcess.StartInfo.RedirectStandardOutput = true;
 				CompilerProcess.StartInfo.RedirectStandardError = true;
-				CompilerProcess.OutputDataReceived += ProcessOutputReceived;
-				CompilerProcess.ErrorDataReceived += ProcessOutputReceived;
+				CompilerProcess.OutputDataReceived += ChmOutputReceived;
+				CompilerProcess.ErrorDataReceived += ChmOutputReceived;
 				CompilerProcess.Start();
 				CompilerProcess.BeginOutputReadLine();
 				CompilerProcess.BeginErrorReadLine();
@@ -1244,6 +1244,21 @@ namespace APIDocTool
 			Utility.SafeCopyFile(Path.Combine(IntermediateDir, ProjectName + ".chm"), ChmFileName);
 		}
 
+		static private void ChmOutputReceived(Object Sender, DataReceivedEventArgs Line)
+		{
+			if(Line.Data != null && Line.Data.Length > 0)
+			{
+				// Building the API docs seems to always give a few of these warnings in random files, but they don't seem to be duplicated.
+				// It looks like HHC hashes the filenames it's seen, and we get a couple of collisions with ~130,000 pages.
+				string OutputLine = Line.Data.TrimEnd();
+				if(OutputLine.EndsWith("is already listed in the [FILES] section of the project file."))
+				{
+					OutputLine = OutputLine.Replace("Warning:", "Note:");
+				}
+				Console.WriteLine(OutputLine);
+			}
+		}
+
 		static private void ProcessOutputReceived(Object Sender, DataReceivedEventArgs Line)
 		{
 			if (Line.Data != null && Line.Data.Length > 0)
