added prereq and engine version metadata to UnrealDocTool and templates

diff --git a/Engine/Source/Programs/UnrealDocTool/UnrealDocTool/MarkdownSharp/Markdown.cs b/Engine/Source/Programs/UnrealDocTool/UnrealDocTool/MarkdownSharp/Markdown.cs
--- a/Engine/Source/Programs/UnrealDocTool/UnrealDocTool/MarkdownSharp/Markdown.cs
+++ b/Engine/Source/Programs/UnrealDocTool/UnrealDocTool/MarkdownSharp/Markdown.cs
@@ -1143,6 +1143,19 @@ namespace MarkdownSharp
                                 relatedPagesCount = data.ProcessedDocumentCache.Metadata.RelatedLinks.Count,
                                 relativeHtmlPath = data.CurrentFolderDetails.RelativeHTMLPath
                             })),
+                        prereqPages = Templates.PrereqPages.Render(Hash.FromAnonymousObject(
+                            new
+                            {
+                                prereqPages = data.ProcessedDocumentCache.Metadata.PrereqLinks,
+                                prereqPagesCount = data.ProcessedDocumentCache.Metadata.PrereqLinks.Count,
+                                relativeHtmlPath = data.CurrentFolderDetails.RelativeHTMLPath
+                            })),
+                        versions = Templates.Versions.Render(Hash.FromAnonymousObject(
+                            new
+                            {
+                                versions = data.ProcessedDocumentCache.Metadata.EngineVersions,
+                                versionCount = data.ProcessedDocumentCache.Metadata.EngineVersions.Count
+                            })),
                         errors = ThisIsPreview
                                      ? Templates.ErrorDetails.Render(Hash.FromAnonymousObject(
                                          new
@@ -1399,6 +1412,35 @@ namespace MarkdownSharp
             return parameters;
         }
 
+        public Hash ProcessPrereqs(string path, TransformationData data)
+        {
+            ClosestFileStatus status;
+            var linkedDoc = data.ProcessedDocumentCache.GetClosest(path, out status);
+
+            if (status == ClosestFileStatus.FileMissing)
+            {
+                var errorId = data.ErrorList.Count;
+                data.ErrorList.Add(
+                    GenerateError(
+                        Language.Message("BadPrereqPageLink", path), MessageClass.Error, path, errorId, data));
+
+                return new Hash();
+            }
+
+            var parameters = Hash.FromAnonymousObject(
+                new
+                {
+                    relativeLink = linkedDoc.GetRelativeTargetPath(data.Document)
+                });
+
+            foreach (var metadata in linkedDoc.PreprocessedData.Metadata.MetadataMap)
+            {
+                parameters.Add(metadata.Key, string.Join(", ", metadata.Value));
+            }
+
+            return parameters;
+        }
+
         /// <summary>
         /// Perform transformations that occur *within* block-level tags like paragraphs, headers, and list items.
         /// </summary>
diff --git a/Engine/Source/Programs/UnrealDocTool/UnrealDocTool/MarkdownSharp/Preprocessor/MetadataManager.cs b/Engine/Source/Programs/UnrealDocTool/UnrealDocTool/MarkdownSharp/Preprocessor/MetadataManager.cs
--- a/Engine/Source/Programs/UnrealDocTool/UnrealDocTool/MarkdownSharp/Preprocessor/MetadataManager.cs
+++ b/Engine/Source/Programs/UnrealDocTool/UnrealDocTool/MarkdownSharp/Preprocessor/MetadataManager.cs
@@ -73,6 +73,16 @@ namespace MarkdownSharp.Preprocessor
                 {
                     RelatedLinks.Add(data.Markdown.ProcessRelated(metaValue, data));
                 }
+
+                if (metaDataCategoryLowerCase == "prereq" && full)
+                {
+                    PrereqLinks.Add(data.Markdown.ProcessPrereqs(metaValue, data));
+                }
+
+                if (metaDataCategoryLowerCase == "version")
+                {
+                    EngineVersions.Add(metaValue);
+                }
             }
 
             // Add meta data to the list, we require some specific meta data keys to be unique others can be duplicates
@@ -133,6 +143,8 @@ namespace MarkdownSharp.Preprocessor
 
             CrumbsLinks = new List<string>();
             RelatedLinks = new List<Hash>();
+            PrereqLinks = new List<Hash>();
+            EngineVersions = new List<String>();
 
             var changes = new List<PreprocessingTextChange>();
 
@@ -193,6 +205,8 @@ namespace MarkdownSharp.Preprocessor
 
         public List<string> CrumbsLinks { get; set; }
         public List<Hash> RelatedLinks { get; set; }
+        public List<Hash> PrereqLinks { get; set; }
+        public List<String> EngineVersions { get; set; }
         public string DocumentTitle { get; set; }
 
         public Dictionary<string, List<string>> MetadataMap { get; private set; }
diff --git a/Engine/Source/Programs/UnrealDocTool/UnrealDocTool/MarkdownSharp/Templates.cs b/Engine/Source/Programs/UnrealDocTool/UnrealDocTool/MarkdownSharp/Templates.cs
--- a/Engine/Source/Programs/UnrealDocTool/UnrealDocTool/MarkdownSharp/Templates.cs
+++ b/Engine/Source/Programs/UnrealDocTool/UnrealDocTool/MarkdownSharp/Templates.cs
@@ -345,6 +345,16 @@ namespace MarkdownSharp
             get { return GetCached("relatedPages.html"); }
         }
 
+        public static TemplateFile PrereqPages
+        {
+            get { return GetCached("prereqPages.html"); }
+        }
+
+        public static TemplateFile Versions
+        {
+            get { return GetCached("Version.html"); }
+        }
+
         public static TemplateFile MissingTemplate
         {
             get { return new TemplateFile(); }
