Updated markdown plugin to support new metadata features

diff --git a/Engine/Source/Programs/UnrealDocTool/MarkdownMode/source.extension.vsixmanifest b/Engine/Source/Programs/UnrealDocTool/MarkdownMode/source.extension.vsixmanifest
--- a/Engine/Source/Programs/UnrealDocTool/MarkdownMode/source.extension.vsixmanifest
+++ b/Engine/Source/Programs/UnrealDocTool/MarkdownMode/source.extension.vsixmanifest
@@ -1,7 +1,7 @@
 <?xml version="1.0" encoding="utf-8"?>
 <PackageManifest Version="2.0.0" xmlns="http://schemas.microsoft.com/developer/vsx-schema/2011" xmlns:d="http://schemas.microsoft.com/developer/vsx-schema-design/2011">
   <Metadata>
-    <Identity Id="Epic MarkdownMode" Version="0.223" Language="en-US" Publisher="Epic Games" />
+    <Identity Id="Epic MarkdownMode" Version="0.224" Language="en-US" Publisher="Epic Games" />
     <DisplayName>Epic MarkdownMode</DisplayName>
     <Description xml:space="preserve">A language service for editing Markdown files.</Description>
     <Icon>markdown-icon.png</Icon>
