#loc UE4DocKOR moved or deleted

diff --git a/Engine/Source/Programs/UnrealDocTool/UnrealDocTranslator/UnrealDocTranslator/UnrealDocTranslatorForm.cs b/Engine/Source/Programs/UnrealDocTool/UnrealDocTranslator/UnrealDocTranslator/UnrealDocTranslatorForm.cs
--- a/Engine/Source/Programs/UnrealDocTool/UnrealDocTranslator/UnrealDocTranslator/UnrealDocTranslatorForm.cs
+++ b/Engine/Source/Programs/UnrealDocTool/UnrealDocTranslator/UnrealDocTranslator/UnrealDocTranslatorForm.cs
@@ -2054,6 +2054,11 @@ namespace UnrealDocTranslator
                                         }
                                         else
                                         {
+                                            if (PartFileName.Contains("ENGINE/PERFORMANCE"))
+                                            {
+                                                int i = 0;
+                                            }
+
                                             INTFileProcessingDictionary.Add(PartFileName,
                                                                             new FileProcessingDetails(
                                                                                 FileDetails.DepotPath.ToString(),
