#loc UE4DocKOR moved or deleted

diff --git a/Engine/Source/Programs/UnrealDocTool/UnrealDocTranslator/UnrealDocTranslator/UnrealDocTranslatorForm.cs b/Engine/Source/Programs/UnrealDocTool/UnrealDocTranslator/UnrealDocTranslator/UnrealDocTranslatorForm.cs
--- a/Engine/Source/Programs/UnrealDocTool/UnrealDocTranslator/UnrealDocTranslator/UnrealDocTranslatorForm.cs
+++ b/Engine/Source/Programs/UnrealDocTool/UnrealDocTranslator/UnrealDocTranslator/UnrealDocTranslatorForm.cs
@@ -1990,6 +1990,12 @@ namespace UnrealDocTranslator
                         {
                             string PartFileName = FileDetails.DepotPath == null ? "" : ParseINTFileName.Match(FileDetails.DepotPath.ToString()).Groups["FileName"].Value.ToUpper();
 
+                            if (PartFileName.Contains("ENGINE/PERFORMANCE"))
+                            {
+                                int i = 0;
+                            }
+
+
                             if (!string.IsNullOrWhiteSpace(PartFileName))
                             {
                                 //Does anyone else have this checked out?
@@ -2003,9 +2009,15 @@ namespace UnrealDocTranslator
                                 }
                                 else
                                 {
-
                                     if (FileDetails.HeadAction == FileAction.Delete)
                                     {
+
+                                        if (PartFileName.Contains("ENGINE/PERFORMANCE"))
+                                        {
+                                            int i = 0;
+                                        }
+
+
                                         INTDeletedFileProcessingDictionary.Add(PartFileName,
                                                                                new FileProcessingDetails(
                                                                                    FileDetails.DepotPath.ToString(),
@@ -2054,11 +2066,6 @@ namespace UnrealDocTranslator
                                         }
                                         else
                                         {
-                                            if (PartFileName.Contains("ENGINE/PERFORMANCE"))
-                                            {
-                                                int i = 0;
-                                            }
-
                                             INTFileProcessingDictionary.Add(PartFileName,
                                                                             new FileProcessingDetails(
                                                                                 FileDetails.DepotPath.ToString(),
