Fix unstable paths for anonymous enums; now just uses the namespace name where possible. Also fix enum values missing from doc pages.

diff --git a/Engine/Source/Programs/UnrealDocTool/APIDocTool/APIDocTool/APIEnum.cs b/Engine/Source/Programs/UnrealDocTool/APIDocTool/APIDocTool/APIEnum.cs
--- a/Engine/Source/Programs/UnrealDocTool/APIDocTool/APIDocTool/APIEnum.cs
+++ b/Engine/Source/Programs/UnrealDocTool/APIDocTool/APIDocTool/APIEnum.cs
@@ -27,7 +27,7 @@ namespace APIDocTool
 		public List<APIEnumValue> Values = new List<APIEnumValue>();
 
 		public APIEnum(APIPage InParent, DoxygenEntity InEntity, string InName)
-			: base(InParent, InName)
+			: base(InParent, GetCleanName(InName))
         {
 			// Set the defaults
 			Entity = InEntity;
@@ -49,6 +49,16 @@ namespace APIDocTool
 			AddRefLink(Node.Attributes["id"].InnerText, this);
         }
 
+		public static string GetCleanName(string Name)
+		{
+			int Index = Name.IndexOf("::@");
+			if(Index != -1)
+			{
+				Name = Name.Substring(0, Index + 2);
+			}
+			return Name;
+		}
+
 		public override void Link()
         {
 			// Read the metadata
@@ -107,6 +117,9 @@ namespace APIDocTool
 				Writer.LeaveSection();
 
 				// Write the enum values
+				Writer.WriteListSection("values", "Values", "Name", "Description", Values.Select(x => x.GetListItem()));
+
+				// Write the description
 				if (!Utility.IsNullOrWhitespace(FullDescription))
 				{
 					Writer.EnterSection("description", "Remarks");
diff --git a/Engine/Source/Programs/UnrealDocTool/APIDocTool/APIDocTool/Utility.cs b/Engine/Source/Programs/UnrealDocTool/APIDocTool/APIDocTool/Utility.cs
--- a/Engine/Source/Programs/UnrealDocTool/APIDocTool/APIDocTool/Utility.cs
+++ b/Engine/Source/Programs/UnrealDocTool/APIDocTool/APIDocTool/Utility.cs
@@ -247,7 +247,7 @@ namespace APIDocTool
 			{
 				Result += "_";
 			}
-			return Result;
+			return Result.TrimEnd('_');
 		}
 
 		public static void Add<Key, Value>(this Dictionary<Key, List<Value>> Map, Key InKey, Value InValue)
