Fix code documentation build error.

diff --git a/Engine/Source/Programs/UnrealDocTool/APIDocTool/APIDocTool/Utility.cs b/Engine/Source/Programs/UnrealDocTool/APIDocTool/APIDocTool/Utility.cs
--- a/Engine/Source/Programs/UnrealDocTool/APIDocTool/APIDocTool/Utility.cs
+++ b/Engine/Source/Programs/UnrealDocTool/APIDocTool/APIDocTool/Utility.cs
@@ -243,11 +243,12 @@ namespace APIDocTool
 					Result += "_";
 				}
 			}
+			Result = Result.TrimEnd('_');
 			if(Result.Equals("CON", StringComparison.CurrentCultureIgnoreCase))
 			{
 				Result += "_";
 			}
-			return Result.TrimEnd('_');
+			return Result;
 		}
 
 		public static void Add<Key, Value>(this Dictionary<Key, List<Value>> Map, Key InKey, Value InValue)
