Use 64-bit doxygen executable by default.

diff --git a/Engine/Source/Programs/UnrealDocTool/APIDocTool/APIDocTool/Program.cs b/Engine/Source/Programs/UnrealDocTool/APIDocTool/APIDocTool/Program.cs
--- a/Engine/Source/Programs/UnrealDocTool/APIDocTool/APIDocTool/Program.cs
+++ b/Engine/Source/Programs/UnrealDocTool/APIDocTool/APIDocTool/Program.cs
@@ -287,7 +287,7 @@ namespace APIDocTool
 			}
 
 			// Derive all the tool paths
-			string DoxygenPath = Path.Combine(EngineDir, "Extras\\NotForLicensees\\Doxygen\\bin\\doxygen.exe");
+			string DoxygenPath = Path.Combine(EngineDir, "Extras\\NotForLicensees\\Doxygen\\bin\\Release64\\doxygen.exe");
 			string EditorPath = Path.Combine(EngineDir, "Binaries\\Win64\\UE4Editor-Cmd.exe");
 			string DocToolPath = Path.Combine(EngineDir, "Binaries\\DotNET\\UnrealDocTool.exe");
 			string ChmCompilerPath = Path.Combine(EngineDir, "Extras\\NotForLicensees\\HTML Help Workshop\\hhc.exe");
