Initial version of API snippet harvesting code. Contains some bug fixes from BenM.

diff --git a/Engine/Source/Programs/UnrealDocTool/APIDocTool/APIDocTool/APIFunction.cs b/Engine/Source/Programs/UnrealDocTool/APIDocTool/APIDocTool/APIFunction.cs
--- a/Engine/Source/Programs/UnrealDocTool/APIDocTool/APIDocTool/APIFunction.cs
+++ b/Engine/Source/Programs/UnrealDocTool/APIDocTool/APIDocTool/APIFunction.cs
@@ -270,6 +270,7 @@ namespace APIDocTool
 		public string FullDescription = "";
 		public string ReturnDescription = "";
 		public List<string> SeeAlso = new List<string>();
+		public List<string> SnippetText = null;
 
 		public APIEventParameters EventParameters;
 
@@ -380,6 +381,9 @@ namespace APIDocTool
 			// Get the @see directives
 			ParseSeeAlso(Node, SeeAlso);
 
+			string RealKeyName = (Entity.Parent != null) ? (Entity.Parent.Name + "::" + Entity.Name) : Entity.Name;
+			SnippetText = APISnippets.LoadSnippetTextForFunction(RealKeyName);
+
 			// Get the modifiers
 			IsVirtual = Node.Attributes.GetNamedItem("virt").InnerText == "virtual";
 			IsStatic = Node.Attributes.GetNamedItem("static").InnerText == "yes";
@@ -718,6 +722,9 @@ namespace APIDocTool
 				// Write the source
 				WriteSourceSection(Writer);
 
+				//Write code snippets
+				WriteSnippetSection(Writer, SnippetText);
+
 				// Write the @see directives
 				WriteSeeAlsoSection(Writer, SeeAlso);
 
diff --git a/Engine/Source/Programs/UnrealDocTool/APIDocTool/APIDocTool/APIMember.cs b/Engine/Source/Programs/UnrealDocTool/APIDocTool/APIDocTool/APIMember.cs
--- a/Engine/Source/Programs/UnrealDocTool/APIDocTool/APIDocTool/APIMember.cs
+++ b/Engine/Source/Programs/UnrealDocTool/APIDocTool/APIDocTool/APIMember.cs
@@ -332,6 +332,38 @@ namespace APIDocTool
 			}
 		}
 
+		public void WriteSnippetSection(UdnWriter Writer, List<string> SnippetLines)
+		{
+			if (SnippetLines != null)
+			{
+				Writer.EnterSection("snippets", "Snippets");
+				Writer.EnterRegion("simplecode");
+				foreach (string SnippetLine in SnippetLines)
+				{
+					StringBuilder OutputLine = new StringBuilder();
+					for (int Idx = 0; Idx < SnippetLine.Length; Idx++)
+					{
+						if (SnippetLine[Idx] == ' ')
+						{
+							OutputLine.Append("&nbsp;");
+						}
+						else if (SnippetLine[Idx] == '\t')
+						{
+							OutputLine.Append("&nbsp;&nbsp;&nbsp;&nbsp;");
+						}
+						else
+						{
+							OutputLine.Append(Markdown.EscapeText(SnippetLine.Substring(Idx)));
+							break;
+						}
+					}
+					Writer.WriteLine(OutputLine.ToString() + "  ");
+				}
+				Writer.LeaveRegion();
+				Writer.LeaveSection();
+			}
+		}
+
 		public static string GetNormalizedFileName(string InFileName)
 		{
 			string NormalizedFileName = InFileName.Replace('\\', '/');
diff --git a/Engine/Source/Programs/UnrealDocTool/APIDocTool/APIDocTool/Program.cs b/Engine/Source/Programs/UnrealDocTool/APIDocTool/APIDocTool/Program.cs
--- a/Engine/Source/Programs/UnrealDocTool/APIDocTool/APIDocTool/Program.cs
+++ b/Engine/Source/Programs/UnrealDocTool/APIDocTool/APIDocTool/Program.cs
@@ -18,6 +18,126 @@ using DoxygenLib;
 
 namespace APIDocTool
 {
+	public class APISnippets
+	{
+		public static string SnippetTextDirectory { get; private set; }
+		public const int FileBufferChunkSize = 2048;
+
+		public APISnippets()
+		{
+			SnippetDictionary = new Dictionary<string, APISnippet>();
+		}
+
+		public static void SetSnippetTextDirectory(string DestinationDirectory)
+		{
+			SnippetTextDirectory = DestinationDirectory;
+		}
+
+		public static void CleanAllFiles()
+		{
+			Directory.Delete(SnippetTextDirectory, true);
+		}
+
+		static string FilenameForKey(string Key)
+		{
+			return (SnippetTextDirectory + "\\" + Key.Replace(':', '-') + ".txt");
+		}
+
+		public static List<string> LoadSnippetTextForFunction(string Key)
+		{
+			string Filename = FilenameForKey(Key);
+			if (!File.Exists(Filename))
+			{
+				return null;
+			}
+			FileInfo fileInfo = new FileInfo(Filename);
+			List<string> returnList = new List<string>(((int)fileInfo.Length / FileBufferChunkSize) + 1);
+			string line;
+			using (FileStream fileStream = File.Open(Filename, FileMode.Open, FileAccess.Read))
+			using (BufferedStream bufferedStream = new BufferedStream(fileStream))
+			using (StreamReader streamReader = new StreamReader(bufferedStream))
+			{
+				while ((line = streamReader.ReadLine()) != null)
+				{
+					returnList.Add(line);
+				}
+			}
+			return returnList;
+		}
+
+		public bool WriteSnippetsToFiles()
+		{
+			if (!Directory.Exists(SnippetTextDirectory))
+			{
+				Directory.CreateDirectory(SnippetTextDirectory);
+			}
+			foreach (string Key in SnippetDictionary.Keys)
+			{
+				using (StreamWriter OutStream = new StreamWriter(FilenameForKey(Key)))
+				{
+					OutStream.AutoFlush = true;
+					SnippetDictionary[Key].WriteToStream(OutStream);
+					OutStream.Close();
+				}
+			}
+			return true;
+		}
+
+		public void AddSnippet(string PageName)
+		{
+			APISnippet ExistingSnippet;
+			if (SnippetDictionary.TryGetValue(PageName, out ExistingSnippet))
+			{
+				//Update this entry with some newlines.
+				ExistingSnippet.AddSnippetText(Environment.NewLine + Environment.NewLine);
+			}
+			else
+			{
+				//Create a new entry.
+				SnippetDictionary.Add(PageName, new APISnippet());
+			}
+		}
+
+		public bool AddSnippetText(string PageName, string NewSnippetText)
+		{
+			APISnippet ExistingSnippet;
+			if (SnippetDictionary.TryGetValue(PageName, out ExistingSnippet))
+			{
+				ExistingSnippet.AddSnippetText(NewSnippetText);
+				return true;
+			}
+			return false;
+		}
+
+		Dictionary<string, APISnippet> SnippetDictionary;
+
+		class APISnippet
+		{
+			public APISnippet()
+			{
+				SnippetText = new List<string>(1);
+			}
+
+			public void AddSnippetText(string NewSnippetText)
+			{
+				SnippetText.Add(NewSnippetText);
+			}
+
+			public void WriteToStream(StreamWriter OutStream)
+			{
+				foreach (string SnippetString in SnippetText)
+				{
+					OutStream.Write(SnippetString);
+					OutStream.Flush();
+				}
+				OutStream.WriteLine();
+			}
+
+			//Multiple snippets can exist per entry. Each list element should be a complete snippet. We can format them, or the space between them, or sort them, after they're all harvested.
+			public List<string> SnippetText { get; private set; }
+		}
+	}
+
     public class Program
 	{
 		[Flags]
@@ -211,9 +331,10 @@ namespace APIDocTool
 			List<string> ArgumentList = new List<string>(Arguments);
 
 			// Get the default paths
-			string EngineDir = ParseArgumentDirectory(ArgumentList, "-enginedir=", Path.Combine(Path.GetDirectoryName(Assembly.GetExecutingAssembly().Location), "..\\..\\..\\..\\..\\..\\..\\.."));
+            string EngineDir = ParseArgumentDirectory(ArgumentList, "-enginedir=", Path.Combine(Path.GetDirectoryName(Assembly.GetExecutingAssembly().Location), "..\\..\\..\\..\\..\\..\\..\\.."));
 			string IntermediateDir = ParseArgumentDirectory(ArgumentList, "-intermediatedir=", Path.Combine(EngineDir, "Intermediate\\Documentation"));
 			string DocumentationDir = ParseArgumentDirectory(ArgumentList, "-documentationdir=", Path.Combine(EngineDir, "Documentation"));
+			string SamplesDir = ParseArgumentDirectory(ArgumentList, "-samplesdir=", Path.Combine(EngineDir, "..\\Templates"));		//Might want to use "..\\Samples" as the path. Consider supporting multiple paths.
 
 			// Check if we're just building an index, no actual content
 			bool bIndexOnly = ArgumentList.Remove("-indexonly");
@@ -242,6 +363,7 @@ namespace APIDocTool
 			}
 
 			BuildActions CodeActions = ParseBuildActions(ArgumentList, "code") | DefaultActions;
+			BuildActions CodeSnippetsActions = ParseBuildActions(ArgumentList, "codesnippets") | CodeActions;
 			BuildActions CodeTargetActions = ParseBuildActions(ArgumentList, "codetarget") | CodeActions;
 			BuildActions CodeMetadataActions = ParseBuildActions(ArgumentList, "codemeta") | CodeActions;
 			BuildActions CodeXmlActions = ParseBuildActions(ArgumentList, "codexml") | CodeActions;
@@ -276,10 +398,11 @@ namespace APIDocTool
 				Console.WriteLine("    -enginedir=<...>:                Specifies the root engine directory");
 				Console.WriteLine("    -intermediatedir=<...>:          Specifies the intermediate directory");
 				Console.WriteLine("    -documentationdir=<...>:         Specifies the documentation directory");
-				Console.WriteLine("    -indexonly:                      Just build index pages");
+                Console.WriteLine("    -samplesdir=<...>:               Specifies the samples directory");
+                Console.WriteLine("    -indexonly:                      Just build index pages");
 				Console.WriteLine("    -filter=<...>,<...>:             Filter which things to convert, eg.");
 				Console.WriteLine("                                     Folders:  -filter=Core/Containers/...");
-				Console.WriteLine("                                     Entities: -filter=Core/TArray");
+				Console.WriteLine("                                     Entities: -filter=Core/TArray (without any folders)");
 				Console.WriteLine("Valid steps are:");
 				Console.WriteLine("   code, codetarget, codemeta, codexml, codeudn, codehtml, codechm");
 				Console.WriteLine("   blueprint, blueprintjson, blueprintudn, blueprinthtml, blueprintchm");
@@ -299,6 +422,7 @@ namespace APIDocTool
 			string MetadataDir = Path.Combine(IntermediateDir, "metadata");
 			string XmlDir = Path.Combine(IntermediateDir, "doxygen");
 			string JsonDir = Path.Combine(IntermediateDir, "json");
+			string SnippetsDir = Path.Combine(IntermediateDir, "snippets");
 			string MetadataPath = Path.Combine(MetadataDir, "metadata.xml");
 
 			// Derive all the output paths
@@ -318,7 +442,14 @@ namespace APIDocTool
 			AddMetadataKeyword(UdnDir, "UPROPERTY", "Programming/UnrealArchitecture/Reference/Properties", "Programming/UnrealArchitecture/Reference/Properties/Specifiers");
 			AddMetadataKeyword(UdnDir, "USTRUCT", "Programming/UnrealArchitecture/Reference/Structs", "Programming/UnrealArchitecture/Reference/Structs/Specifiers");
 
+			// Establish snippet directory so we can look things up later
+			APISnippets.SetSnippetTextDirectory(SnippetsDir);
+
 			// Build all the code docs
+			if (!BuildCodeSnippetsTxt(SamplesDir, CodeSnippetsActions))
+			{
+				return 1;
+			}
 			if (!BuildCodeTargetInfo(TargetInfoPath, EngineDir, Path.Combine(ArchiveDir, "CodeAPI-TargetInfo.tgz"), CodeTargetActions))
 			{
 				return 1;
@@ -572,131 +703,242 @@ namespace APIDocTool
 			return false;
 		}
 
-		static bool BuildCodeXml(string EngineDir, string TargetInfoPath, string DoxygenPath, string XmlDir, string ArchivePath, List<string> Filters, BuildActions Actions)
+		static bool BuildCodeSnippetsTxt(string SamplesDir, BuildActions Actions)
 		{
 			if ((Actions & BuildActions.Clean) != 0)
 			{
-				Console.WriteLine("Cleaning '{0}'", XmlDir);
-				Utility.SafeDeleteDirectoryContents(XmlDir, true);
+				APISnippets.CleanAllFiles();
 			}
-			if ((Actions & BuildActions.Build) != 0)
+			if ((Actions & BuildActions.Build) == 0)
 			{
-				Console.WriteLine("Building XML...");
-				Utility.SafeCreateDirectory(XmlDir);
-
-				// Read the target that we're building
-				BuildTarget Target = new BuildTarget(Path.Combine(EngineDir, "Source"), TargetInfoPath);
-
-				// Create an invariant list of exclude directories
-				string[] InvariantExcludeDirectories = ExcludeSourceDirectories.Select(x => x.ToLowerInvariant()).ToArray();
+				return true;
+			}
 
-				// Get the list of folders to filter against
-				List<string> FolderFilters = new List<string>();
-				if (Filters != null)
+			//We should be able to trim this down further.
+			DirectoryInfo Dir = new DirectoryInfo(SamplesDir);
+			List<string> Files = new List<string>();
+			foreach (string FileName in Directory.GetFiles(SamplesDir, "*.cpp", SearchOption.AllDirectories))
+			{
+				if (!FileName.EndsWith(".generated.cpp"))
 				{
-					foreach (string Filter in Filters)
-					{
-						int Idx = Filter.IndexOf('/');
-						if (Idx != -1)
-						{
-							FolderFilters.Add("\\" + Filter.Substring(0, Idx) + "\\");
-						}
-					}
+					Files.Add(FileName);
 				}
-
-				// Flatten the target into a list of modules
-				List<string> InputModules = new List<string>();
-				foreach (string DirectoryName in Target.DirectoryNames)
+			}
+			foreach (string FileName in Directory.GetFiles(SamplesDir, "*.h", SearchOption.AllDirectories))
+			{
+				if (!FileName.EndsWith(".generated.h"))
 				{
-					for (DirectoryInfo ModuleDirectory = new DirectoryInfo(DirectoryName); ModuleDirectory.Parent != null; ModuleDirectory = ModuleDirectory.Parent)
-					{
-						IEnumerator<FileInfo> ModuleFile = ModuleDirectory.EnumerateFiles("*.build.cs").GetEnumerator();
-						if (ModuleFile.MoveNext() && (FolderFilters.Count == 0 || FolderFilters.Any(x => ModuleFile.Current.FullName.Contains(x))))
-						{
-							InputModules.AddUnique(ModuleFile.Current.FullName);
-							break;
-						}
-					}
+					Files.Add(FileName);
 				}
+			}
 
-				// Just use all the input modules
-				if (!bIndexOnly)
+			//Do the harvesting work.
+			{
+				const string OpeningTag = "///CODE_SNIPPET_START:";
+				const string ClosingTag = "///CODE_SNIPPET_END";
+				APISnippets Snippets = new APISnippets();
+				List<string> CurrentSnippetPageNames = new List<string>(4);		//Probably won't have a snippet that is shared by more than four different pages.
+				string CurrentLine;
+				char[] WhiteSpace = {' ', '\t'};		//Doubles as our token delimiter in one place - noted in comments.
+				bool IsSnippetBeingProcessed = false;
+
+				foreach (string FileName in Files)
 				{
-					// Set our error mode so as to not bring up the WER dialog if Doxygen crashes (which it often does)
-					SetErrorMode(0x0007);
-
-					// Create the output directory
-					Utility.SafeCreateDirectory(XmlDir);
-
-					// Build all the modules
-					Console.WriteLine("Parsing source...");
-
-					// Build the list of definitions
-					List<string> Definitions = new List<string>();
-					foreach (string Definition in Target.Definitions)
+					// Read the file and display it line by line.
+					System.IO.StreamReader file = new System.IO.StreamReader(FileName);
+					while ((CurrentLine = file.ReadLine()) != null)
 					{
-						if (!Definition.StartsWith("ORIGINAL_FILE_NAME="))
+						CurrentLine = CurrentLine.TrimStart(WhiteSpace);
+						if (!CurrentLine.StartsWith(OpeningTag))
 						{
-							Definitions.Add(Definition.Replace("DLLIMPORT", "").Replace("DLLEXPORT", ""));
+							continue;
+						}
+						CurrentSnippetPageNames = CurrentLine.Split(WhiteSpace).ToList<string>();		//Whitespace is used to delimit our API/snippet page names here.
+						CurrentSnippetPageNames.RemoveAt(0);											//Remove the opening tag, which is always in position 0 after empties have been cleared out.
+						CurrentSnippetPageNames.RemoveAll(entry => (entry.Length < 1));					//Blank entries can show up in the list. Remove them.
+						if (CurrentSnippetPageNames.Count < 1)
+						{
+							Console.WriteLine("Warning: OpeningTag for snippet harvesting found without any API pages specified.");
+							continue;
 						}
-					}
 
-					// Build a list of input paths
-					List<string> InputPaths = new List<string>();
-					foreach (string InputModule in InputModules)
-					{
-						foreach (string DirectoryName in Directory.EnumerateDirectories(Path.GetDirectoryName(InputModule), "*", SearchOption.AllDirectories))
+						IsSnippetBeingProcessed = true;
+						foreach (string CurrentSnippetPageName in CurrentSnippetPageNames)
 						{
-							// Find the relative path from the engine directory
-							string NormalizedDirectoryName = DirectoryName;
-							if (NormalizedDirectoryName.StartsWith(EngineDir))
+							Snippets.AddSnippet(CurrentSnippetPageName);
+						}
+						while ((CurrentLine = file.ReadLine()) != null)
+						{
+							string TrimmedLine = CurrentLine.TrimStart(WhiteSpace);		//This is actually a C++ same-line whitespace check, not our token delimiters.
+							if (TrimmedLine.StartsWith(OpeningTag))
 							{
-								NormalizedDirectoryName = NormalizedDirectoryName.Substring(EngineDir.Length);
+								//Snippets do not currently support overlapping. If they did, closing tags should be explicit about which entries are ending, and we'd need to skip lines with opening tags.
+								Console.WriteLine("Error: Nested OpeningTag found! This is not supported. Snippet harvesting process will fail.");
+								return false;
 							}
-							if (!NormalizedDirectoryName.EndsWith("\\"))
+							else if (TrimmedLine.StartsWith(ClosingTag))
 							{
-								NormalizedDirectoryName += "\\";
+								//We're done with this snippet now. Mark that we can end cleanly.
+								IsSnippetBeingProcessed = false;
+								break;
 							}
 
-							// Check we can include it
-							if (!ExcludeSourceDirectories.Any(x => NormalizedDirectoryName.Contains("\\" + x + "\\")))
+							//This line should be added to the snippet(s) named in the "CODE_SNIPPET_START" line. Capture it. We need to add our own newline.
+							foreach (string CurrentSnippetPageName in CurrentSnippetPageNames)
 							{
-								if (FolderFilters.Count == 0 || FolderFilters.Any(x => NormalizedDirectoryName.Contains(x)))
+								if (!Snippets.AddSnippetText(CurrentSnippetPageName, CurrentLine + Environment.NewLine))
 								{
-									InputPaths.Add(DirectoryName);
+									Console.WriteLine("Error adding text to snippet for " + CurrentSnippetPageName);
+									return false;
 								}
 							}
 						}
 					}
-
-					// Build the configuration for this module
-					DoxygenConfig Config = new DoxygenConfig("UE4", InputPaths.ToArray(), XmlDir);
-					Config.Definitions.AddRange(Definitions);
-					Config.Definitions.AddRange(DoxygenPredefinedMacros);
-					Config.ExpandAsDefined.AddRange(DoxygenExpandedMacros);
-					Config.IncludePaths.AddRange(Target.IncludePaths);
-					Config.ExcludePatterns.AddRange(ExcludeSourceDirectories.Select(x => "*/" + x + "/*"));
-					Config.ExcludePatterns.AddRange(ExcludeSourceFiles);
-
-					// Run Doxygen
-					if (!Doxygen.Run(DoxygenPath, Path.Combine(EngineDir, "Source"), Config, true))
+					//If we hit the end of the file while harvesting a snippet, we should fail or have a warning. We could also just ignore the failure to end cleanly and just store the text as-is.
+					//Opting for outright failure at the moment so that these errors don't go unnoticed.
+					if (IsSnippetBeingProcessed)
 					{
-						Console.WriteLine("  Doxygen crashed. Skipping.");
+						Console.WriteLine("Code snippet start tag not matched with code snippet end tag in " + FileName);
 						return false;
 					}
+					file.Close();
 				}
 
-				// Write the modules file
-				File.WriteAllLines(Path.Combine(XmlDir, "modules.txt"), InputModules);
-			}
-			if((Actions & BuildActions.Archive) != 0)
-			{
-				Console.WriteLine("Creating archive '{0}'", ArchivePath);
-				Utility.CreateTgzFromDir(ArchivePath, XmlDir);
+				if (!Snippets.WriteSnippetsToFiles())
+				{
+					Console.WriteLine("Error writing intermediate snippet files.");
+					return false;
+				}
 			}
+			//Completed without error.
 			return true;
 		}
 
+		static bool BuildCodeXml(string EngineDir, string TargetInfoPath, string DoxygenPath, string XmlDir, string ArchivePath, List<string> Filters, BuildActions Actions)
+        {
+            if ((Actions & BuildActions.Clean) != 0)
+            {
+                Console.WriteLine("Cleaning '{0}'", XmlDir);
+                Utility.SafeDeleteDirectoryContents(XmlDir, true);
+            }
+            if ((Actions & BuildActions.Build) != 0)
+            {
+                Console.WriteLine("Building XML...");
+                Utility.SafeCreateDirectory(XmlDir);
+
+                // Read the target that we're building
+                BuildTarget Target = new BuildTarget(Path.Combine(EngineDir, "Source"), TargetInfoPath);
+
+                // Create an invariant list of exclude directories
+                string[] InvariantExcludeDirectories = ExcludeSourceDirectories.Select(x => x.ToLowerInvariant()).ToArray();
+
+                // Get the list of folders to filter against
+                List<string> FolderFilters = new List<string>();
+                if (Filters != null)
+                {
+                    foreach (string Filter in Filters)
+                    {
+                        int Idx = Filter.IndexOf('/');
+                        if (Idx != -1)
+                        {
+                            FolderFilters.Add("\\" + Filter.Substring(0, Idx) + "\\");
+                        }
+                    }
+                }
+
+                // Flatten the target into a list of modules
+                List<string> InputModules = new List<string>();
+                foreach (string DirectoryName in Target.DirectoryNames)
+                {
+                    for (DirectoryInfo ModuleDirectory = new DirectoryInfo(DirectoryName); ModuleDirectory.Parent != null; ModuleDirectory = ModuleDirectory.Parent)
+                    {
+                        IEnumerator<FileInfo> ModuleFile = ModuleDirectory.EnumerateFiles("*.build.cs").GetEnumerator();
+                        if (ModuleFile.MoveNext() && (FolderFilters.Count == 0 || FolderFilters.Any(x => ModuleFile.Current.FullName.Contains(x))))
+                        {
+                            InputModules.AddUnique(ModuleFile.Current.FullName);
+                            break;
+                        }
+                    }
+                }
+
+                // Just use all the input modules
+                if (!bIndexOnly)
+                {
+                    // Set our error mode so as to not bring up the WER dialog if Doxygen crashes (which it often does)
+                    SetErrorMode(0x0007);
+
+                    // Create the output directory
+                    Utility.SafeCreateDirectory(XmlDir);
+
+                    // Build all the modules
+                    Console.WriteLine("Parsing source...");
+
+                    // Build the list of definitions
+                    List<string> Definitions = new List<string>();
+                    foreach (string Definition in Target.Definitions)
+                    {
+                        if (!Definition.StartsWith("ORIGINAL_FILE_NAME="))
+                        {
+                            Definitions.Add(Definition.Replace("DLLIMPORT", "").Replace("DLLEXPORT", ""));
+                        }
+                    }
+
+                    // Build a list of input paths
+                    List<string> InputPaths = new List<string>();
+                    foreach (string InputModule in InputModules)
+                    {
+                        foreach (string DirectoryName in Directory.EnumerateDirectories(Path.GetDirectoryName(InputModule), "*", SearchOption.AllDirectories))
+                        {
+                            // Find the relative path from the engine directory
+                            string NormalizedDirectoryName = DirectoryName;
+                            if (NormalizedDirectoryName.StartsWith(EngineDir))
+                            {
+                                NormalizedDirectoryName = NormalizedDirectoryName.Substring(EngineDir.Length);
+                            }
+                            if (!NormalizedDirectoryName.EndsWith("\\"))
+                            {
+                                NormalizedDirectoryName += "\\";
+                            }
+
+                            // Check we can include it
+                            if (!ExcludeSourceDirectories.Any(x => NormalizedDirectoryName.Contains("\\" + x + "\\")))
+                            {
+                                if (FolderFilters.Count == 0 || FolderFilters.Any(x => NormalizedDirectoryName.Contains(x)))
+                                {
+                                    InputPaths.Add(DirectoryName);
+                                }
+                            }
+                        }
+                    }
+
+                    // Build the configuration for this module
+                    DoxygenConfig Config = new DoxygenConfig("UE4", InputPaths.ToArray(), XmlDir);
+                    Config.Definitions.AddRange(Definitions);
+                    Config.Definitions.AddRange(DoxygenPredefinedMacros);
+                    Config.ExpandAsDefined.AddRange(DoxygenExpandedMacros);
+                    Config.IncludePaths.AddRange(Target.IncludePaths);
+                    Config.ExcludePatterns.AddRange(ExcludeSourceDirectories.Select(x => "*/" + x + "/*"));
+                    Config.ExcludePatterns.AddRange(ExcludeSourceFiles);
+
+                    // Run Doxygen
+                    if (!Doxygen.Run(DoxygenPath, Path.Combine(EngineDir, "Source"), Config, true))
+                    {
+                        Console.WriteLine("  Doxygen crashed. Skipping.");
+                        return false;
+                    }
+                }
+
+                // Write the modules file
+                File.WriteAllLines(Path.Combine(XmlDir, "modules.txt"), InputModules);
+            }
+            if ((Actions & BuildActions.Archive) != 0)
+            {
+                Console.WriteLine("Creating archive '{0}'", ArchivePath);
+                Utility.CreateTgzFromDir(ArchivePath, XmlDir);
+            }
+            return true;
+        }
+
 		static bool BuildBlueprintJson(string JsonDir, string EngineDir, string EditorPath, string ArchivePath, BuildActions Actions)
 		{
 			if((Actions & BuildActions.Clean) != 0)
