Bugfix: Added a check to snippet clean procedure to ensure the intermediate directory exists before attempting to delete it. Also added "codesnippets" option to the help output.

diff --git a/Engine/Source/Programs/UnrealDocTool/APIDocTool/APIDocTool/Program.cs b/Engine/Source/Programs/UnrealDocTool/APIDocTool/APIDocTool/Program.cs
--- a/Engine/Source/Programs/UnrealDocTool/APIDocTool/APIDocTool/Program.cs
+++ b/Engine/Source/Programs/UnrealDocTool/APIDocTool/APIDocTool/Program.cs
@@ -35,7 +35,10 @@ namespace APIDocTool
 
 		public static void CleanAllFiles()
 		{
-			Directory.Delete(SnippetTextDirectory, true);
+			if (Directory.Exists(SnippetTextDirectory))
+			{
+				Directory.Delete(SnippetTextDirectory, true);
+			}
 		}
 
 		static string FilenameForKey(string Key)
@@ -389,7 +392,7 @@ namespace APIDocTool
 				Console.WriteLine("Options:");
 				Console.WriteLine("    -rebuild:                        Clean and build everything");
 				Console.WriteLine("    -rebuild<step>:                  Clean and build specific steps");
-				Console.WriteLine("    -clean:                          Clean all files");
+				Console.WriteLine("    -clean:                          Clean all files"); 
 				Console.WriteLine("    -clean<step>:                    Clean specific steps");
 				Console.WriteLine("    -build:                          Build everything");
 				Console.WriteLine("    -build<step>:                    Build specific output steps");
@@ -404,7 +407,7 @@ namespace APIDocTool
 				Console.WriteLine("                                     Folders:  -filter=Core/Containers/...");
 				Console.WriteLine("                                     Entities: -filter=Core/TArray (without any folders)");
 				Console.WriteLine("Valid steps are:");
-				Console.WriteLine("   code, codetarget, codemeta, codexml, codeudn, codehtml, codechm");
+				Console.WriteLine("   code, codetarget, codemeta, codexml, codeudn, codehtml, codechm, codesnippets");
 				Console.WriteLine("   blueprint, blueprintjson, blueprintudn, blueprinthtml, blueprintchm");
 				return 1;
 			}
