Created example snippet-harvesting block for GetActorLocation and GetActorRotation in sample ShooterGame. Changed the default snippet path from Templates to Samples.

diff --git a/Engine/Source/Programs/UnrealDocTool/APIDocTool/APIDocTool/Program.cs b/Engine/Source/Programs/UnrealDocTool/APIDocTool/APIDocTool/Program.cs
--- a/Engine/Source/Programs/UnrealDocTool/APIDocTool/APIDocTool/Program.cs
+++ b/Engine/Source/Programs/UnrealDocTool/APIDocTool/APIDocTool/Program.cs
@@ -337,7 +337,7 @@ namespace APIDocTool
             string EngineDir = ParseArgumentDirectory(ArgumentList, "-enginedir=", Path.Combine(Path.GetDirectoryName(Assembly.GetExecutingAssembly().Location), "..\\..\\..\\..\\..\\..\\..\\.."));
 			string IntermediateDir = ParseArgumentDirectory(ArgumentList, "-intermediatedir=", Path.Combine(EngineDir, "Intermediate\\Documentation"));
 			string DocumentationDir = ParseArgumentDirectory(ArgumentList, "-documentationdir=", Path.Combine(EngineDir, "Documentation"));
-			string SamplesDir = ParseArgumentDirectory(ArgumentList, "-samplesdir=", Path.Combine(EngineDir, "..\\Templates"));		//Might want to use "..\\Samples" as the path. Consider supporting multiple paths.
+			string SamplesDir = ParseArgumentDirectory(ArgumentList, "-samplesdir=", Path.Combine(EngineDir, "..\\Samples"));
 
 			// Check if we're just building an index, no actual content
 			bool bIndexOnly = ArgumentList.Remove("-indexonly");
