Fix build error in APIDocTool

diff --git a/Engine/Source/Programs/UnrealDocTool/APIDocTool/APIDocTool/Program.cs b/Engine/Source/Programs/UnrealDocTool/APIDocTool/APIDocTool/Program.cs
--- a/Engine/Source/Programs/UnrealDocTool/APIDocTool/APIDocTool/Program.cs
+++ b/Engine/Source/Programs/UnrealDocTool/APIDocTool/APIDocTool/Program.cs
@@ -299,7 +299,7 @@ namespace APIDocTool
 
 			// Delegates
 			"FUNC_PARAM_MEMBERS=",
-			"DECLARE_DERIVED_EVENT(X,Y,Z)="
+			"DECLARE_DERIVED_EVENT(X,Y,Z)=",
 			"DELEGATE_DEPRECATED(X)=",
 		};
 
