Removed source-file text from documentation.

diff --git a/Engine/Source/Programs/UnrealDocTool/APIDocTool/APIDocTool/APIFunction.cs b/Engine/Source/Programs/UnrealDocTool/APIDocTool/APIDocTool/APIFunction.cs
--- a/Engine/Source/Programs/UnrealDocTool/APIDocTool/APIDocTool/APIFunction.cs
+++ b/Engine/Source/Programs/UnrealDocTool/APIDocTool/APIDocTool/APIFunction.cs
@@ -566,61 +566,6 @@ namespace APIDocTool
 			WriteNestedSimpleCode(Writer, Lines);
 		}
 
-		private void WriteSourceSection(UdnWriter Writer)
-		{
-			if(Entity.BodyFile != null)
-			{
-				DoxygenSourceFile SourceFile = Entity.Module.FindSourceFile(Entity.BodyFile);
-				if(SourceFile != null)
-				{
-					int BodyStart = Math.Min(Math.Max(Entity.BodyStart - 1, 0), SourceFile.Lines.Count - 1);
-					int BodyEnd = Math.Min(Math.Max(Entity.BodyEnd, BodyStart), SourceFile.Lines.Count);
-					if(BodyEnd > BodyStart)
-					{
-						Writer.EnterSection("source", "Source");
-						Writer.EnterRegion("simplecode");
-
-						List<string> Lines = new List<string>();
-						int MinPrefix = int.MaxValue;
-
-						for (int LineIdx = BodyStart; LineIdx < BodyEnd; LineIdx++)
-						{
-							XmlNode Node = SourceFile.Lines[LineIdx];
-							string MarkdownLine = (Node == null)? "" : Markdown.ParseXmlCodeLine(Node, ResolveDoxygenLink);
-
-							int Prefix = 0;
-							while (Prefix < MarkdownLine.Length && MarkdownLine[Prefix] == ' ') Prefix++;
-
-							if(Prefix < MarkdownLine.Length && Prefix < MinPrefix)
-							{
-								MinPrefix = Prefix;
-							}
-
-							Lines.Add(MarkdownLine);
-						}
-
-						for (int Idx = 0; Idx < Lines.Count; Idx++)
-						{
-							int TextIdx = Math.Min(MinPrefix, Lines[Idx].Length);
-							if(TextIdx == Lines[Idx].Length)
-							{
-								Writer.Write("&nbsp;");
-							}
-							while(TextIdx < Lines[Idx].Length && Lines[Idx][TextIdx] == ' ')
-							{
-								Writer.Write("&nbsp;");
-								TextIdx++;
-							}
-							Writer.WriteLine(Lines[Idx].Substring(TextIdx) + "  ");
-						}
-
-						Writer.LeaveRegion();
-						Writer.LeaveSection();
-					}
-				}
-			}
-		}
-
 		private void WriteIcons(UdnWriter Writer)
 		{
 			Writer.WriteIcon(Icons.Function[(int)Protection]);
@@ -719,9 +664,6 @@ namespace APIDocTool
 					Writer.LeaveSection();
 				}
 
-				// Write the source
-				WriteSourceSection(Writer);
-
 				//Write code snippets
 				WriteSnippetSection(Writer, SnippetText);
 
diff --git a/Engine/Source/Programs/UnrealDocTool/APIDocTool/APIDocTool/Program.cs b/Engine/Source/Programs/UnrealDocTool/APIDocTool/APIDocTool/Program.cs
--- a/Engine/Source/Programs/UnrealDocTool/APIDocTool/APIDocTool/Program.cs
+++ b/Engine/Source/Programs/UnrealDocTool/APIDocTool/APIDocTool/Program.cs
@@ -1034,11 +1034,6 @@ namespace APIDocTool
 						Entity.Module = Module;
 						Module.Entities.Add(Entity);
 					}
-					foreach (DoxygenSourceFile SourceFile in RootModule.SourceFiles)
-					{
-						DoxygenModule Module = Modules.Find(x => SourceFile.FileName.Replace('/', '\\').StartsWith(x.BaseSrcDir));
-						Module.SourceFiles.Add(SourceFile);
-					}
 
 					// Now filter all the entities in each module
 					if (Filters != null && Filters.Count > 0)
diff --git a/Engine/Source/Programs/UnrealDocTool/DoxygenLib/DoxygenModule.cs b/Engine/Source/Programs/UnrealDocTool/DoxygenLib/DoxygenModule.cs
--- a/Engine/Source/Programs/UnrealDocTool/DoxygenLib/DoxygenModule.cs
+++ b/Engine/Source/Programs/UnrealDocTool/DoxygenLib/DoxygenModule.cs
@@ -106,53 +106,12 @@ namespace DoxygenLib
 		}
 	}
 
-	public class DoxygenSourceFile
-	{
-		public string FileName;
-		public string NormalizedFileName;
-		public List<XmlNode> Lines = new List<XmlNode>();
-
-		public DoxygenSourceFile(string InFileName)
-		{
-			FileName = InFileName;
-			NormalizedFileName = NormalizeFileName(InFileName);
-		}
-
-		public static string NormalizeFileName(string InFileName)
-		{
-			return InFileName.Replace('\\', '/').ToLowerInvariant();
-		}
-
-		public static DoxygenSourceFile FromXml(XmlNode Node)
-		{
-			XmlNode LocationNode = Node.SelectSingleNode("location");
-			if(LocationNode == null) return null;
-
-			DoxygenSourceFile SourceFile = new DoxygenSourceFile(LocationNode.Attributes["file"].InnerText);
-			foreach(XmlNode CodeLineNode in Node.SelectNodes("programlisting/codeline"))
-			{
-				int CodeLineIdx = int.Parse(CodeLineNode.Attributes["lineno"].InnerText) - 1;
-				while(CodeLineIdx >= SourceFile.Lines.Count)
-				{
-					SourceFile.Lines.Add(null);
-				}
-				SourceFile.Lines[CodeLineIdx] = CodeLineNode;
-			}
-			return SourceFile;
-		}
-
-		public override string ToString()
-		{
-			return FileName;
-		}
-	}
 
 	public class DoxygenModule
 	{
 		public readonly string Name;
 		public readonly string BaseSrcDir;
 		public List<DoxygenEntity> Entities = new List<DoxygenEntity>();
-		public List<DoxygenSourceFile> SourceFiles = new List<DoxygenSourceFile>();
 
 		public DoxygenModule(string InName, string InBaseSrcDir)
 		{
@@ -207,13 +166,7 @@ namespace DoxygenLib
 			Dictionary<string, DoxygenEntity> Scopes = new Dictionary<string, DoxygenEntity>();
 			foreach (DoxygenCompound Compound in Compounds)
 			{
-				if (Compound.Kind == "file")
-				{
-					ReadMembers(Module, Compound, "", null, null, Module.Entities);
-					DoxygenSourceFile SourceFile = DoxygenSourceFile.FromXml(Compound.Node);
-					if(SourceFile != null) Module.SourceFiles.Add(SourceFile);
-				}
-				else if (Compound.Kind == "namespace")
+				if (Compound.Kind == "namespace")
 				{
 					ReadMembers(Module, Compound, Compound.Name + "::", null, Compound.Node, Module.Entities);
 				}
@@ -265,23 +218,6 @@ namespace DoxygenLib
 			}
 		}
 
-		protected static DoxygenSourceFile ReadSourceFile(string BaseXmlDir, string CompoundId)
-		{
-			string XmlFileName = Path.Combine(BaseXmlDir, CompoundId + ".xml");
-
-			XmlDocument Document;
-			if(TryReadXmlDocument(XmlFileName, out Document))
-			{
-				DoxygenSourceFile SourceFile = DoxygenSourceFile.FromXml(Document.SelectSingleNode("doxygen/compounddef"));
-				return SourceFile;
-			}
-			else
-			{
-				Console.WriteLine("Couldn't read source document: '{0}'", XmlFileName);
-				return null;
-			}
-		}
-
 		protected static void ReadMembers(DoxygenModule Module, DoxygenCompound Compound, string NamePrefix, DoxygenEntity Parent, XmlNode NamespaceNode, List<DoxygenEntity> Entities)
 		{
 			using (XmlNodeList NodeList = Compound.Node.SelectNodes("sectiondef/memberdef"))
@@ -297,6 +233,7 @@ namespace DoxygenLib
 			}
 		}
 
+		/*
 		public DoxygenSourceFile FindSourceFile(string InFileName)
 		{
 			string NormalizedFileName = DoxygenSourceFile.NormalizeFileName(InFileName);
@@ -309,6 +246,7 @@ namespace DoxygenLib
 			}
 			return null;
 		}
+		 * */
 
 		public override string ToString()
 		{
