Fix for blueprint docs build failing due to Blutility module not being loaded, and add editor command line into the log.

diff --git a/Engine/Source/Editor/UnrealEd/Private/Commandlets/GenerateBlueprintAPICommandlet.cpp b/Engine/Source/Editor/UnrealEd/Private/Commandlets/GenerateBlueprintAPICommandlet.cpp
--- a/Engine/Source/Editor/UnrealEd/Private/Commandlets/GenerateBlueprintAPICommandlet.cpp
+++ b/Engine/Source/Editor/UnrealEd/Private/Commandlets/GenerateBlueprintAPICommandlet.cpp
@@ -972,6 +972,8 @@ int32 UGenerateBlueprintAPICommandlet::Main(FString const& Params)
 	GenerateBlueprintAPIUtils::CommandOptions = GenerateBlueprintAPIUtils::CommandletOptions(Switches);
 	GenerateBlueprintAPIUtils::CommandletOptions const& CommandOptions = GenerateBlueprintAPIUtils::CommandOptions;
 
+	FModuleManager::Get().LoadModuleChecked(TEXT("Blutility"));
+
 	FString ActiveFilePath;
 	FArchive* FileOut = nullptr;
 
diff --git a/Engine/Source/Programs/UnrealDocTool/APIDocTool/APIDocTool/Program.cs b/Engine/Source/Programs/UnrealDocTool/APIDocTool/APIDocTool/Program.cs
--- a/Engine/Source/Programs/UnrealDocTool/APIDocTool/APIDocTool/Program.cs
+++ b/Engine/Source/Programs/UnrealDocTool/APIDocTool/APIDocTool/Program.cs
@@ -956,11 +956,14 @@ namespace APIDocTool
 				// Create the output directory
 				Utility.SafeCreateDirectory(JsonDir);
 
+				string Arguments = "-run=GenerateBlueprintAPI -path=" + JsonDir + " -name=BlueprintAPI";
+				Console.WriteLine("Running: {0} {1}", EditorPath, Arguments);
+
 				using (Process JsonExportProcess = new Process())
 				{
 					JsonExportProcess.StartInfo.WorkingDirectory = EngineDir;
 					JsonExportProcess.StartInfo.FileName = EditorPath;
-					JsonExportProcess.StartInfo.Arguments = "-run=GenerateBlueprintAPI -path=" + JsonDir + " -name=BlueprintAPI";
+					JsonExportProcess.StartInfo.Arguments = Arguments;
 					JsonExportProcess.StartInfo.UseShellExecute = false;
 					JsonExportProcess.StartInfo.RedirectStandardOutput = true;
 					JsonExportProcess.StartInfo.RedirectStandardError = true;
