Fixed a bug in UnrealDocTool that caused a crash if an invalid source directory was provided.

diff --git a/Engine/Source/Programs/UnrealDocTool/UnrealDocTool/UnrealDocTool/Program.cs b/Engine/Source/Programs/UnrealDocTool/UnrealDocTool/UnrealDocTool/Program.cs
--- a/Engine/Source/Programs/UnrealDocTool/UnrealDocTool/UnrealDocTool/Program.cs
+++ b/Engine/Source/Programs/UnrealDocTool/UnrealDocTool/UnrealDocTool/Program.cs
@@ -1074,7 +1074,12 @@ namespace UnrealDocTool
 
         private static void DoRecursively(string sourcePath, Action<string> action)
         {
-            action(sourcePath);
+			if (!Directory.Exists(sourcePath))
+			{
+				return;
+			}
+
+			action(sourcePath);
 
             foreach (var subDirectory in Directory.GetDirectories(sourcePath))
             {
