Back out changelist 2429454
Submitted unintentionally by UnrealDocTranslator

diff --git a/Engine/Source/Programs/UnrealDocTool/UnrealDocTranslator/UnrealDocTranslator/UnrealDocTranslatorForm.cs b/Engine/Source/Programs/UnrealDocTool/UnrealDocTranslator/UnrealDocTranslator/UnrealDocTranslatorForm.cs
--- a/Engine/Source/Programs/UnrealDocTool/UnrealDocTranslator/UnrealDocTranslator/UnrealDocTranslatorForm.cs
+++ b/Engine/Source/Programs/UnrealDocTool/UnrealDocTranslator/UnrealDocTranslator/UnrealDocTranslatorForm.cs
@@ -2059,11 +2059,6 @@ namespace UnrealDocTranslator
                                     }
                                     else
                                     {
-                                        if (PartFileName.Contains("ENGINE/PERFORMANCE"))
-                                        {
-                                            int i = 0;
-                                        }
-
                                         if (FileDetails.ClientPath == null)
                                         {
                                             //log message
