Back out revision 20 from //depot/UE4/Engine/Source/Programs/UnrealDocTool/UnrealDocTranslator/UnrealDocTranslator/UnrealDocTranslatorForm.cs
Unintentional submission from UnrealDocTranslator tool

diff --git a/Engine/Source/Programs/UnrealDocTool/UnrealDocTranslator/UnrealDocTranslator/UnrealDocTranslatorForm.cs b/Engine/Source/Programs/UnrealDocTool/UnrealDocTranslator/UnrealDocTranslator/UnrealDocTranslatorForm.cs
--- a/Engine/Source/Programs/UnrealDocTool/UnrealDocTranslator/UnrealDocTranslator/UnrealDocTranslatorForm.cs
+++ b/Engine/Source/Programs/UnrealDocTool/UnrealDocTranslator/UnrealDocTranslator/UnrealDocTranslatorForm.cs
@@ -1990,12 +1990,6 @@ namespace UnrealDocTranslator
                         {
                             string PartFileName = FileDetails.DepotPath == null ? "" : ParseINTFileName.Match(FileDetails.DepotPath.ToString()).Groups["FileName"].Value.ToUpper();
 
-                            if (PartFileName.Contains("ENGINE/PERFORMANCE"))
-                            {
-                                int i = 0;
-                            }
-
-
                             if (!string.IsNullOrWhiteSpace(PartFileName))
                             {
                                 //Does anyone else have this checked out?
@@ -2009,15 +2003,9 @@ namespace UnrealDocTranslator
                                 }
                                 else
                                 {
+
                                     if (FileDetails.HeadAction == FileAction.Delete)
                                     {
-
-                                        if (PartFileName.Contains("ENGINE/PERFORMANCE"))
-                                        {
-                                            int i = 0;
-                                        }
-
-
                                         INTDeletedFileProcessingDictionary.Add(PartFileName,
                                                                                new FileProcessingDetails(
                                                                                    FileDetails.DepotPath.ToString(),
@@ -2066,6 +2054,11 @@ namespace UnrealDocTranslator
                                         }
                                         else
                                         {
+                                            if (PartFileName.Contains("ENGINE/PERFORMANCE"))
+                                            {
+                                                int i = 0;
+                                            }
+
                                             INTFileProcessingDictionary.Add(PartFileName,
                                                                             new FileProcessingDetails(
                                                                                 FileDetails.DepotPath.ToString(),
