UnrealDocTranslator: Fixed unintentional submissions bug
- Previously checked to make sure you didn't have anything from Documentation\\Source folder open for edit in default changelist before submitting moves/deletes
- Now checks to make sure you don't have files from _any_ folder open for edit in default changelist before submitting

diff --git a/Engine/Source/Programs/UnrealDocTool/UnrealDocTranslator/UnrealDocTranslator/UnrealDocTranslatorForm.cs b/Engine/Source/Programs/UnrealDocTool/UnrealDocTranslator/UnrealDocTranslator/UnrealDocTranslatorForm.cs
--- a/Engine/Source/Programs/UnrealDocTool/UnrealDocTranslator/UnrealDocTranslator/UnrealDocTranslatorForm.cs
+++ b/Engine/Source/Programs/UnrealDocTool/UnrealDocTranslator/UnrealDocTranslator/UnrealDocTranslatorForm.cs
@@ -1941,13 +1941,14 @@ namespace UnrealDocTranslator
                         HeadSourceUdnFiles[0] = new FileSpec(
                             new DepotPath(_connectionDetails.DepotPath), null, null, VersionSpec.Head);
 
-                        FileSpec[] OpenedSourceUdnFiles = new FileSpec[1];
-                        OpenedSourceUdnFiles[0] = new FileSpec(
-                            new DepotPath(_connectionDetails.DepotPath), null, null, null);
+                        FileSpec[] OpenedFiles = new FileSpec[1];
+                        OpenedFiles[0] = new FileSpec(
+                            new DepotPath("..."), null, null, null);    // Make sure there are no files in any path open for edit in the default changelist
+                                                                        // They will be unintentionally submitted otherwise.
 
                         try
                         {
-                            var ListOfOpenedFileStats = P4Repository.GetOpenedFiles(OpenedSourceUdnFiles, null);
+                            var ListOfOpenedFileStats = P4Repository.GetOpenedFiles(OpenedFiles, null);
 
                             FilesCheckedOut.Clear();
 
