UnrealDocTranslator: Fix missing files in UnrealDocTranslator
- Fixed issue where if the last perforce action for a file was a MoveAdd, it was not showing up in the UI

diff --git a/Engine/Source/Programs/UnrealDocTool/UnrealDocTranslator/UnrealDocTranslator/UnrealDocTranslatorForm.cs b/Engine/Source/Programs/UnrealDocTool/UnrealDocTranslator/UnrealDocTranslator/UnrealDocTranslatorForm.cs
--- a/Engine/Source/Programs/UnrealDocTool/UnrealDocTranslator/UnrealDocTranslator/UnrealDocTranslatorForm.cs
+++ b/Engine/Source/Programs/UnrealDocTool/UnrealDocTranslator/UnrealDocTranslator/UnrealDocTranslatorForm.cs
@@ -2044,6 +2044,17 @@ namespace UnrealDocTranslator
                                                                                          FileProcessingDetails.State.MoveTo,
                                                                                          IsCheckedOutByOtherUser,
                                                                                          _connectionDetails));
+
+                                            // Even if the most recent perforce action is a MoveAdd, 
+                                            // We need to add this to the file processing dictionary as well, so it can show up in the tool's UI
+                                            INTFileProcessingDictionary.Add(PartFileName,
+                                                                            new FileProcessingDetails(
+                                                                                FileDetails.DepotPath.ToString(),
+                                                                                FileDetails.ClientPath.ToString(),
+                                                                                FileDetails.HeadChange,
+                                                                                FileProcessingDetails.State.NotDeterminedYet,
+                                                                                IsCheckedOutByOtherUser,
+                                                                                _connectionDetails));
                                         }
                                     }
                                     else
