Pass appropriate parameters to the spawned editor for generating blueprint JSON file.

diff --git a/Engine/Source/Programs/AutomationTool/Scripts/BuildDocumentation.Automation.cs b/Engine/Source/Programs/AutomationTool/Scripts/BuildDocumentation.Automation.cs
--- a/Engine/Source/Programs/AutomationTool/Scripts/BuildDocumentation.Automation.cs
+++ b/Engine/Source/Programs/AutomationTool/Scripts/BuildDocumentation.Automation.cs
@@ -51,7 +51,7 @@ class DocumentationNode : GUBP.GUBPNode
 	protected void ExecuteApiDocTool(string Arguments, string LogName)
 	{
 		string ApiDocToolPath = Path.Combine(CommandUtils.CmdEnv.LocalRoot, "Engine/Source/Programs/UnrealDocTool/APIDocTool/APIDocTool/bin/x64/Release/APIDocTool.exe");
-		string ApiToolCommandLine = Arguments + " -enginedir=\"" + Path.Combine(CommandUtils.CmdEnv.LocalRoot, "Engine") + "\"";
+		string ApiToolCommandLine = Arguments + " -enginedir=\"" + Path.Combine(CommandUtils.CmdEnv.LocalRoot, "Engine") + "\"" + (CommandUtils.IsBuildMachine? " -buildmachine" : "");
 		CommandUtils.RunAndLog(CommandUtils.CmdEnv, ApiDocToolPath, ApiToolCommandLine, LogName);
 	}
 
diff --git a/Engine/Source/Programs/UnrealDocTool/APIDocTool/APIDocTool/Program.cs b/Engine/Source/Programs/UnrealDocTool/APIDocTool/APIDocTool/Program.cs
--- a/Engine/Source/Programs/UnrealDocTool/APIDocTool/APIDocTool/Program.cs
+++ b/Engine/Source/Programs/UnrealDocTool/APIDocTool/APIDocTool/Program.cs
@@ -344,6 +344,9 @@ namespace APIDocTool
 			// Check if we're just building an index, no actual content
 			bool bIndexOnly = ArgumentList.Remove("-indexonly");
 
+			// Check if we're running on a build machine
+			bool bBuildMachine = ArgumentList.Remove("-buildmachine");
+
 			// Remove all the filter arguments
 			List<string> Filters = new List<string>();
 			for (int Idx = 0; Idx < ArgumentList.Count; Idx++)
@@ -481,7 +484,7 @@ namespace APIDocTool
 			}
 
 			// Build all the blueprint docs
-			if (!BuildBlueprintJson(JsonDir, EngineDir, EditorPath, Path.Combine(ArchiveDir, "BlueprintAPI-JSON.tgz"), BlueprintJsonActions))
+			if (!BuildBlueprintJson(JsonDir, EngineDir, EditorPath, Path.Combine(ArchiveDir, "BlueprintAPI-JSON.tgz"), bBuildMachine, BlueprintJsonActions))
 			{
 				return 1;
 			}
@@ -944,7 +947,7 @@ namespace APIDocTool
             return true;
         }
 
-		static bool BuildBlueprintJson(string JsonDir, string EngineDir, string EditorPath, string ArchivePath, BuildActions Actions)
+		static bool BuildBlueprintJson(string JsonDir, string EngineDir, string EditorPath, string ArchivePath, bool bBuildMachine, BuildActions Actions)
 		{
 			if((Actions & BuildActions.Clean) != 0)
 			{
@@ -956,7 +959,7 @@ namespace APIDocTool
 				// Create the output directory
 				Utility.SafeCreateDirectory(JsonDir);
 
-				string Arguments = "-run=GenerateBlueprintAPI -path=" + JsonDir + " -name=BlueprintAPI";
+				string Arguments = "-run=GenerateBlueprintAPI -path=" + JsonDir + " -name=BlueprintAPI -stdout -FORCELOGFLUSH -CrashForUAT -unattended -AllowStdOutLogVerbosity" + (bBuildMachine? " -buildmachine" : "");
 				Console.WriteLine("Running: {0} {1}", EditorPath, Arguments);
 
 				using (Process JsonExportProcess = new Process())
