[INTEGRATE] Change 2442372 by [REDACTED_PERSON][REDACTED_HANDLE] on 2015/02/11 16:58:16

	[UE-9282] - Fixed by adding the capability to separate notes from normal text. Output is now  interleaved between the two text types.

diff --git a/Engine/Source/Programs/UnrealDocTool/APIDocTool/APIDocTool/APIAction.cs b/Engine/Source/Programs/UnrealDocTool/APIDocTool/APIDocTool/APIAction.cs
--- a/Engine/Source/Programs/UnrealDocTool/APIDocTool/APIDocTool/APIAction.cs
+++ b/Engine/Source/Programs/UnrealDocTool/APIDocTool/APIDocTool/APIAction.cs
@@ -10,7 +10,36 @@ namespace APIDocTool
 {
 	class APIAction : APIPage
 	{
-		public readonly string Tooltip;
+		protected class TooltipLine
+		{
+			public TooltipLine()
+			{
+			}
+
+			public TooltipLine(LineType InType)
+			{
+				Type = InType;
+			}
+
+			public TooltipLine(LineType InType, string InText)
+			{
+				Type = InType;
+				Text = InText;
+			}
+
+			public enum LineType
+			{
+				Normal,
+				Note,
+				Count
+			}
+
+			public LineType Type = LineType.Count;
+			public string Text;
+		}
+
+		protected readonly List<TooltipLine> TooltipData = new List<TooltipLine>();
+		public readonly string TooltipNormalText;
 		public readonly string CompactName;
 		public readonly string NodeType;
 
@@ -22,14 +51,39 @@ namespace APIDocTool
 			: base(InParent, InName)
 		{
 			object Value;
-			
+
+			TooltipNormalText = "";
+			TooltipLine.LineType CurrentLineType = TooltipLine.LineType.Count;		//Invalid
 			if (ActionProperties.TryGetValue("Tooltip", out Value))
 			{
-				Tooltip = String.Concat(((string)Value).Split('\n').Select(Line => Line.Trim() + '\n'));
-			}
-			else
-			{
-				Tooltip = "";
+				//Create an interleaved list of normal text and note regions. Also, store all normal text as a single block (i.e. notes stripped out) in a separate place.
+				foreach (string Line in ((string)Value).Split('\n'))
+				{
+					string TrimmedLine = Line.Trim();
+
+					if (TrimmedLine.StartsWith("@note"))
+					{
+						if (TrimmedLine.Length > 6)
+						{
+							if (CurrentLineType != TooltipLine.LineType.Note)
+							{
+								CurrentLineType = TooltipLine.LineType.Note;
+								TooltipData.Add(new TooltipLine(CurrentLineType));
+							}
+							TooltipData[TooltipData.Count - 1].Text += (TrimmedLine.Substring(6) + '\n');
+						}
+					}
+					else
+					{
+						if (CurrentLineType != TooltipLine.LineType.Normal)
+						{
+							CurrentLineType = TooltipLine.LineType.Normal;
+							TooltipData.Add(new TooltipLine(CurrentLineType));
+						}
+						TooltipData[TooltipData.Count - 1].Text += (TrimmedLine + '\n');
+						TooltipNormalText += (TrimmedLine + '\n');
+					}
+				}
 			}
 
 			if (ActionProperties.TryGetValue("CompactTitle", out Value))
@@ -84,8 +138,24 @@ namespace APIDocTool
 			{
 				Writer.WritePageHeader(Name, PageCrumbs, Name);
 
-				// Tooltip/description
-				Writer.WriteLine(Tooltip);
+				// Tooltip/description - Write this as interleaved text and notes
+				foreach (TooltipLine TTL in TooltipData)
+				{
+					switch (TTL.Type)
+					{
+						case TooltipLine.LineType.Normal:
+							Writer.WriteLine(TTL.Text);
+							break;
+						case TooltipLine.LineType.Note:
+							Writer.EnterRegion("note");
+							Writer.WriteLine(TTL.Text);
+							Writer.LeaveRegion();
+							break;
+						default:
+							//Error? Ignore this entry for now.
+							break;
+					}
+				}
 
 				if (Pins != null && Pins.Count() > 0)
 				{
@@ -160,7 +230,7 @@ namespace APIDocTool
 
 		public UdnListItem GetListItem()
 		{
-			return new UdnListItem(Name, Tooltip, LinkPath);
+			return new UdnListItem(Name, TooltipNormalText, LinkPath);
 		}
 	}
 }
