Merging improvement to Code Snippets from 4.7 Branch (original CL 2456709).

diff --git a/Engine/Source/Programs/UnrealDocTool/APIDocTool/APIDocTool/APIMember.cs b/Engine/Source/Programs/UnrealDocTool/APIDocTool/APIDocTool/APIMember.cs
--- a/Engine/Source/Programs/UnrealDocTool/APIDocTool/APIDocTool/APIMember.cs
+++ b/Engine/Source/Programs/UnrealDocTool/APIDocTool/APIDocTool/APIMember.cs
@@ -336,30 +336,11 @@ namespace APIDocTool
 		{
 			if (SnippetLines != null)
 			{
-				Writer.EnterSection("snippets", "Snippets");
-				Writer.EnterRegion("simplecode");
+				Writer.EnterSection("examplecode", "Example Code");
 				foreach (string SnippetLine in SnippetLines)
 				{
-					StringBuilder OutputLine = new StringBuilder();
-					for (int Idx = 0; Idx < SnippetLine.Length; Idx++)
-					{
-						if (SnippetLine[Idx] == ' ')
-						{
-							OutputLine.Append("&nbsp;");
-						}
-						else if (SnippetLine[Idx] == '\t')
-						{
-							OutputLine.Append("&nbsp;&nbsp;&nbsp;&nbsp;");
-						}
-						else
-						{
-							OutputLine.Append(Markdown.EscapeText(SnippetLine.Substring(Idx)));
-							break;
-						}
-					}
-					Writer.WriteLine(OutputLine.ToString() + "  ");
+					Writer.WriteLine("\t" + SnippetLine + "  ");
 				}
-				Writer.LeaveRegion();
 				Writer.LeaveSection();
 			}
 		}
diff --git a/Engine/Source/Programs/UnrealDocTool/APIDocTool/APIDocTool/Program.cs b/Engine/Source/Programs/UnrealDocTool/APIDocTool/APIDocTool/Program.cs
--- a/Engine/Source/Programs/UnrealDocTool/APIDocTool/APIDocTool/Program.cs
+++ b/Engine/Source/Programs/UnrealDocTool/APIDocTool/APIDocTool/Program.cs
@@ -749,6 +749,7 @@ namespace APIDocTool
 				string CurrentLine;
 				char[] WhiteSpace = {' ', '\t'};		//Doubles as our token delimiter in one place - noted in comments.
 				bool IsSnippetBeingProcessed = false;
+				bool WasPreviousLineBlank = false;		//Two blank lines in a row will end a code block. Don't allow this to happen in a single snippet.
 
 				foreach (string FileName in Files)
 				{
@@ -777,7 +778,7 @@ namespace APIDocTool
 						}
 						while ((CurrentLine = file.ReadLine()) != null)
 						{
-							string TrimmedLine = CurrentLine.TrimStart(WhiteSpace);		//This is actually a C++ same-line whitespace check, not our token delimiters.
+							string TrimmedLine = CurrentLine.TrimStart(WhiteSpace);		//This is actually a C# same-line whitespace check, not our token delimiters.
 							if (TrimmedLine.StartsWith(OpeningTag))
 							{
 								//Snippets do not currently support overlapping. If they did, closing tags should be explicit about which entries are ending, and we'd need to skip lines with opening tags.
@@ -790,6 +791,22 @@ namespace APIDocTool
 								IsSnippetBeingProcessed = false;
 								break;
 							}
+							if (CurrentLine.Trim().Length < 1)
+							{
+								if (WasPreviousLineBlank)
+								{
+									//Two (or more) blank lines in a row! Not permitted.
+									continue;
+								}
+								else
+								{
+									WasPreviousLineBlank = true;
+								}
+							}
+							else
+							{
+								WasPreviousLineBlank = false;
+							}
 
 							//This line should be added to the snippet(s) named in the "CODE_SNIPPET_START" line. Capture it. We need to add our own newline.
 							foreach (string CurrentSnippetPageName in CurrentSnippetPageNames)
