Added safety checks on delete functions in APIDocTool

diff --git a/Engine/Source/Programs/UnrealDocTool/APIDocTool/APIDocTool/Utility.cs b/Engine/Source/Programs/UnrealDocTool/APIDocTool/APIDocTool/Utility.cs
--- a/Engine/Source/Programs/UnrealDocTool/APIDocTool/APIDocTool/Utility.cs
+++ b/Engine/Source/Programs/UnrealDocTool/APIDocTool/APIDocTool/Utility.cs
@@ -59,8 +59,10 @@ namespace APIDocTool
 			}
 		}
 
-		public static void SafeDeleteFile(string Path)
+		// Returns false if delete fails
+		public static bool SafeDeleteFile(string Path)
 		{
+			bool bResult = true;
 			if (File.Exists(Path))
 			{
 				FileAttributes Attributes = File.GetAttributes(Path);
@@ -68,22 +70,36 @@ namespace APIDocTool
 				{
 					File.SetAttributes(Path, Attributes & ~FileAttributes.ReadOnly);
 				}
-				File.Delete(Path);
+				try
+				{
+					File.Delete(Path);
+				}
+				catch (IOException Ex)
+				{					
+					Console.WriteLine("Failed to delete '{0}':\n{1}", Path, Ex);
+					bResult = false;
+				}
 			}
+			return bResult;
 		}
 
-		public static void SafeDeleteDirectory(string Dir)
+		// Returns false if delete fails
+		public static bool SafeDeleteDirectory(string Dir)
 		{
+			bool bResult = true;
 			if (Directory.Exists(Dir))
 			{
 				try
 				{
 					Directory.Delete(Dir, true);
 				}
-				catch(IOException)
+				catch (IOException Ex)
 				{
+					Console.WriteLine("Failed to delete directory '{0}':\n{1}", Path, Ex);
+					bResult = false;
 				}
 			}
+			return bResult;
 		}
 
 		public static void SafeDeleteDirectoryContents(string Dir, bool bIncludeSubDirectories)
