Fixed compile error in APIDocTool Utility.cs

diff --git a/Engine/Source/Programs/UnrealDocTool/APIDocTool/APIDocTool/Utility.cs b/Engine/Source/Programs/UnrealDocTool/APIDocTool/APIDocTool/Utility.cs
--- a/Engine/Source/Programs/UnrealDocTool/APIDocTool/APIDocTool/Utility.cs
+++ b/Engine/Source/Programs/UnrealDocTool/APIDocTool/APIDocTool/Utility.cs
@@ -95,7 +95,7 @@ namespace APIDocTool
 				}
 				catch (IOException Ex)
 				{
-					Console.WriteLine("Failed to delete directory '{0}':\n{1}", Path, Ex);
+					Console.WriteLine("Failed to delete directory '{0}':\n{1}", Dir, Ex);
 					bResult = false;
 				}
 			}
