[INTEGRATE] Change 2453101 by bruce.nesbit@bnesbit_Releases on 2015/02/20 08:51:44

	Added support for snippets to variables doc generation.
	Renamed LoadSnippetTextForFunction to LoadSnippetTextForSymbol

diff --git a/Engine/Source/Programs/UnrealDocTool/APIDocTool/APIDocTool/APIFunction.cs b/Engine/Source/Programs/UnrealDocTool/APIDocTool/APIDocTool/APIFunction.cs
--- a/Engine/Source/Programs/UnrealDocTool/APIDocTool/APIDocTool/APIFunction.cs
+++ b/Engine/Source/Programs/UnrealDocTool/APIDocTool/APIDocTool/APIFunction.cs
@@ -382,7 +382,7 @@ namespace APIDocTool
 			ParseSeeAlso(Node, SeeAlso);
 
 			string RealKeyName = (Entity.Parent != null) ? (Entity.Parent.Name + "::" + Entity.Name) : Entity.Name;
-			SnippetText = APISnippets.LoadSnippetTextForFunction(RealKeyName);
+			SnippetText = APISnippets.LoadSnippetTextForSymbol(RealKeyName);
 
 			// Get the modifiers
 			IsVirtual = Node.Attributes.GetNamedItem("virt").InnerText == "virtual";
diff --git a/Engine/Source/Programs/UnrealDocTool/APIDocTool/APIDocTool/APIVariable.cs b/Engine/Source/Programs/UnrealDocTool/APIDocTool/APIDocTool/APIVariable.cs
--- a/Engine/Source/Programs/UnrealDocTool/APIDocTool/APIDocTool/APIVariable.cs
+++ b/Engine/Source/Programs/UnrealDocTool/APIDocTool/APIDocTool/APIVariable.cs
@@ -35,6 +35,7 @@ namespace APIDocTool
 		public bool IsStatic = false;
 
 		public List<string> Warnings = new List<string>();
+		public List<string> SnippetText = null;
 
 		public APIVariable(APIPage InParent, XmlNode InNode)
 			: base(InParent, InNode.SelectSingleNode("name").InnerText)
@@ -100,6 +101,8 @@ namespace APIDocTool
 						break;
 				}
 			}
+
+			SnippetText = APISnippets.LoadSnippetTextForSymbol(FullName);			
 		}
 
 		private void WriteDefinition(UdnWriter Writer)
@@ -153,6 +156,9 @@ namespace APIDocTool
 					Writer.WriteLine(FullDescription);
 					Writer.LeaveSection();
 				}
+
+				//Write code snippets
+				WriteSnippetSection(Writer, SnippetText);
 			}
         }
 
diff --git a/Engine/Source/Programs/UnrealDocTool/APIDocTool/APIDocTool/Program.cs b/Engine/Source/Programs/UnrealDocTool/APIDocTool/APIDocTool/Program.cs
--- a/Engine/Source/Programs/UnrealDocTool/APIDocTool/APIDocTool/Program.cs
+++ b/Engine/Source/Programs/UnrealDocTool/APIDocTool/APIDocTool/Program.cs
@@ -41,14 +41,14 @@ namespace APIDocTool
 			}
 		}
 
-		static string FilenameForKey(string Key)
+		static string SnippetFilenameForKey(string Key)
 		{
 			return (SnippetTextDirectory + "\\" + Key.Replace(':', '-') + ".txt");
 		}
 
-		public static List<string> LoadSnippetTextForFunction(string Key)
+		public static List<string> LoadSnippetTextForSymbol(string Key)
 		{
-			string Filename = FilenameForKey(Key);
+			string Filename = SnippetFilenameForKey(Key);
 			if (!File.Exists(Filename))
 			{
 				return null;
@@ -76,7 +76,7 @@ namespace APIDocTool
 			}
 			foreach (string Key in SnippetDictionary.Keys)
 			{
-				using (StreamWriter OutStream = new StreamWriter(FilenameForKey(Key)))
+				using (StreamWriter OutStream = new StreamWriter(SnippetFilenameForKey(Key)))
 				{
 					OutStream.AutoFlush = true;
 					SnippetDictionary[Key].WriteToStream(OutStream);
