Nested code-snippet opening tags in APIDocTool are now ignored with a warning instead of causing the process to fail.

diff --git a/Engine/Source/Programs/UnrealDocTool/APIDocTool/APIDocTool/Program.cs b/Engine/Source/Programs/UnrealDocTool/APIDocTool/APIDocTool/Program.cs
--- a/Engine/Source/Programs/UnrealDocTool/APIDocTool/APIDocTool/Program.cs
+++ b/Engine/Source/Programs/UnrealDocTool/APIDocTool/APIDocTool/Program.cs
@@ -782,8 +782,9 @@ namespace APIDocTool
 							if (TrimmedLine.StartsWith(OpeningTag))
 							{
 								//Snippets do not currently support overlapping. If they did, closing tags should be explicit about which entries are ending, and we'd need to skip lines with opening tags.
-								Console.WriteLine("Error: Nested OpeningTag found! This is not supported. Snippet harvesting process will fail.");
-								return false;
+								//We are currently issuing warnings and ignoring nested opening tags. We do not attempt to match a nested opening tag to a nested closing tag.
+								Console.WriteLine("Warning: Nested OpeningTag found! This is not supported. Snippet harvesting process will ignore this line.");
+								continue;
 							}
 							else if (TrimmedLine.StartsWith(ClosingTag))
 							{
