[UE-11052] Not done, but made some improvements.

diff --git a/Engine/Source/Programs/UnrealDocTool/APIDocTool/APIDocTool/APIMember.cs b/Engine/Source/Programs/UnrealDocTool/APIDocTool/APIDocTool/APIMember.cs
--- a/Engine/Source/Programs/UnrealDocTool/APIDocTool/APIDocTool/APIMember.cs
+++ b/Engine/Source/Programs/UnrealDocTool/APIDocTool/APIDocTool/APIMember.cs
@@ -339,7 +339,7 @@ namespace APIDocTool
 				Writer.EnterSection("examplecode", "Example Code");
 				foreach (string SnippetLine in SnippetLines)
 				{
-					Writer.WriteLine("\t" + SnippetLine + "  ");
+					Writer.WriteLine(SnippetLine + "  ");
 				}
 				Writer.LeaveSection();
 			}
diff --git a/Engine/Source/Programs/UnrealDocTool/APIDocTool/APIDocTool/Program.cs b/Engine/Source/Programs/UnrealDocTool/APIDocTool/APIDocTool/Program.cs
--- a/Engine/Source/Programs/UnrealDocTool/APIDocTool/APIDocTool/Program.cs
+++ b/Engine/Source/Programs/UnrealDocTool/APIDocTool/APIDocTool/Program.cs
@@ -755,36 +755,51 @@ namespace APIDocTool
 				{
 					// Read the file and display it line by line.
 					System.IO.StreamReader file = new System.IO.StreamReader(FileName);
+					int LineNumber = 0;
 					while ((CurrentLine = file.ReadLine()) != null)
 					{
+						++LineNumber;
 						CurrentLine = CurrentLine.TrimStart(WhiteSpace);
-						if (!CurrentLine.StartsWith(OpeningTag))
+						if (!CurrentLine.Contains(OpeningTag))
 						{
 							continue;
 						}
 						CurrentSnippetPageNames = CurrentLine.Split(WhiteSpace).ToList<string>();		//Whitespace is used to delimit our API/snippet page names here.
-						CurrentSnippetPageNames.RemoveAt(0);											//Remove the opening tag, which is always in position 0 after empties have been cleared out.
 						CurrentSnippetPageNames.RemoveAll(entry => (entry.Length < 1));					//Blank entries can show up in the list. Remove them.
+						while ((CurrentSnippetPageNames.Count > 0) && (!CurrentSnippetPageNames[0].Contains(OpeningTag)))
+						{
+							CurrentSnippetPageNames.RemoveAt(0);											//Remove everything before the opening tag.
+						}
+						if (CurrentSnippetPageNames.Count > 0)
+						{
+							CurrentSnippetPageNames.RemoveAt(0);											//Remove the opening tag, which is always in position 0 by this point.
+						}
 						if (CurrentSnippetPageNames.Count < 1)
 						{
-							Console.WriteLine("Warning: OpeningTag for snippet harvesting found without any API pages specified.");
-							continue;
+							Console.WriteLine("Error: OpeningTag for snippet harvesting found without any API pages specified.");
+							return false;
 						}
 
 						IsSnippetBeingProcessed = true;
 						foreach (string CurrentSnippetPageName in CurrentSnippetPageNames)
 						{
 							Snippets.AddSnippet(CurrentSnippetPageName);
+							if (!Snippets.AddSnippetText(CurrentSnippetPageName, "**" + Path.GetFileName(FileName) + "** at line " + LineNumber + ":" + Environment.NewLine + Environment.NewLine))
+							{
+								Console.WriteLine("Error: Failed to add header text to snippet for " + CurrentSnippetPageName + " from source file " + FileName + " at line " + LineNumber);
+								return false;
+							}
 						}
+						bool FinishSnippetAfterThisLine = false;
 						while ((CurrentLine = file.ReadLine()) != null)
 						{
+							++LineNumber;
 							string TrimmedLine = CurrentLine.TrimStart(WhiteSpace);		//This is actually a C# same-line whitespace check, not our token delimiters.
-							if (TrimmedLine.StartsWith(OpeningTag))
+							if (TrimmedLine.Contains(OpeningTag))
 							{
 								//Snippets do not currently support overlapping. If they did, closing tags should be explicit about which entries are ending, and we'd need to skip lines with opening tags.
-								//We are currently issuing warnings and ignoring nested opening tags. We do not attempt to match a nested opening tag to a nested closing tag.
-								Console.WriteLine("Warning: Nested OpeningTag found! This is not supported. Snippet harvesting process will ignore this line.");
-								continue;
+								Console.WriteLine("Error: Nested OpeningTag found in " + FileName + " at line " + LineNumber + "! This is not supported. Snippet harvesting process will ignore this line.");
+								return false;
 							}
 							else if (TrimmedLine.StartsWith(ClosingTag))
 							{
@@ -792,6 +807,13 @@ namespace APIDocTool
 								IsSnippetBeingProcessed = false;
 								break;
 							}
+							else if (TrimmedLine.Contains(ClosingTag))
+							{
+								//We will be done this snippet once we add this line (minus the tag).
+								FinishSnippetAfterThisLine = true;
+								CurrentLine.Replace(ClosingTag, "//");
+							}
+
 							if (CurrentLine.Trim().Length < 1)
 							{
 								if (WasPreviousLineBlank)
@@ -812,12 +834,17 @@ namespace APIDocTool
 							//This line should be added to the snippet(s) named in the "CODE_SNIPPET_START" line. Capture it. We need to add our own newline.
 							foreach (string CurrentSnippetPageName in CurrentSnippetPageNames)
 							{
-								if (!Snippets.AddSnippetText(CurrentSnippetPageName, CurrentLine + Environment.NewLine))
+								if (!Snippets.AddSnippetText(CurrentSnippetPageName, '\t' + CurrentLine + Environment.NewLine))
 								{
-									Console.WriteLine("Error adding text to snippet for " + CurrentSnippetPageName);
+									Console.WriteLine("Error: Failed to add text to snippet for " + CurrentSnippetPageName + " from source file " + FileName + " at line " + LineNumber);
 									return false;
 								}
 							}
+							if (FinishSnippetAfterThisLine)
+							{
+								IsSnippetBeingProcessed = false;
+								break;
+							}
 						}
 					}
 					//If we hit the end of the file while harvesting a snippet, we should fail or have a warning. We could also just ignore the failure to end cleanly and just store the text as-is.
@@ -1557,14 +1584,6 @@ namespace APIDocTool
 				}
 			}
 		}
-		/*
-		public static string ParseXML(XmlNode ParentNode, XmlNode ChildNode, string Indent)
-		{
-//			string Indent = "";
-//			for(int Idx = 0; Idx < TabDepth; Idx++) Indent += "\t";
-			return Markdown.ParseXml(ChildNode, Indent, ResolveLink);
-		}
-		*/
 		public static string ResolveLink(string Id)
 		{
 			APIMember RefMember = APIMember.ResolveRefLink(Id);
