"See Also" section added to APIDocTool with ///CODE_SNIPPET_SEE: tag, pending a better tag name.

diff --git a/Engine/Source/Programs/UnrealDocTool/APIDocTool/APIDocTool/Program.cs b/Engine/Source/Programs/UnrealDocTool/APIDocTool/APIDocTool/Program.cs
--- a/Engine/Source/Programs/UnrealDocTool/APIDocTool/APIDocTool/Program.cs
+++ b/Engine/Source/Programs/UnrealDocTool/APIDocTool/APIDocTool/Program.cs
@@ -101,6 +101,19 @@ namespace APIDocTool
 			}
 		}
 
+		public void FinishCurrentSnippet(string PageName)
+		{
+			APISnippet ExistingSnippet;
+			if (SnippetDictionary.TryGetValue(PageName, out ExistingSnippet))
+			{
+				ExistingSnippet.FinishCurrentSnippet();
+			}
+			else
+			{
+				Console.WriteLine("WARNING: Called FinishCurrentSnippet() with PageName " + PageName + ", but no snippet by that name was found.");
+			}
+		}
+
 		public bool AddSnippetText(string PageName, string NewSnippetText)
 		{
 			APISnippet ExistingSnippet;
@@ -112,6 +125,17 @@ namespace APIDocTool
 			return false;
 		}
 
+		public bool AddSeeText(string PageName, string NewSeeText)
+		{
+			APISnippet ExistingSnippet;
+			if (SnippetDictionary.TryGetValue(PageName, out ExistingSnippet))
+			{
+				ExistingSnippet.AddSeeText(NewSeeText);
+				return true;
+			}
+			return false;
+		}
+
 		Dictionary<string, APISnippet> SnippetDictionary;
 
 		class APISnippet
@@ -119,6 +143,7 @@ namespace APIDocTool
 			public APISnippet()
 			{
 				SnippetText = new List<string>(1);
+				PendingSeeText = new List<string>(4);
 			}
 
 			public void AddSnippetText(string NewSnippetText)
@@ -126,6 +151,24 @@ namespace APIDocTool
 				SnippetText.Add(NewSnippetText);
 			}
 
+			public void AddSeeText(string NewSeeText)
+			{
+				PendingSeeText.Add(NewSeeText);
+			}
+
+			public void FinishCurrentSnippet()
+			{
+				if (PendingSeeText.Count > 0)
+				{
+					AddSnippetText("See Also:" + Environment.NewLine);
+					foreach (string SeeString in PendingSeeText)
+					{
+						AddSnippetText(SeeString);
+					}
+					PendingSeeText.Clear();
+				}
+			}
+
 			public void WriteToStream(StreamWriter OutStream)
 			{
 				foreach (string SnippetString in SnippetText)
@@ -138,6 +181,9 @@ namespace APIDocTool
 
 			//Multiple snippets can exist per entry. Each list element should be a complete snippet. We can format them, or the space between them, or sort them, after they're all harvested.
 			public List<string> SnippetText { get; private set; }
+
+			//See entries are written after the snippet code, so they're stored separately for convenience.
+			public List<string> PendingSeeText { get; private set; }
 		}
 	}
 
@@ -744,8 +790,10 @@ namespace APIDocTool
 			{
 				const string OpeningTag = "///CODE_SNIPPET_START:";
 				const string ClosingTag = "///CODE_SNIPPET_END";
+				const string SeeTag = "///CODE_SNIPPET_SEE:";
 				APISnippets Snippets = new APISnippets();
 				List<string> CurrentSnippetPageNames = new List<string>(4);		//Probably won't have a snippet that is shared by more than four different pages.
+				List<string> SeePageNames = new List<string>(4);				//More than four of these on one line will probably not happen.
 				string CurrentLine;
 				char[] WhiteSpace = {' ', '\t'};		//Doubles as our token delimiter in one place - noted in comments.
 				bool IsSnippetBeingProcessed = false;
@@ -798,20 +846,53 @@ namespace APIDocTool
 							if (TrimmedLine.Contains(OpeningTag))
 							{
 								//Snippets do not currently support overlapping. If they did, closing tags should be explicit about which entries are ending, and we'd need to skip lines with opening tags.
-								Console.WriteLine("Error: Nested OpeningTag found in " + FileName + " at line " + LineNumber + "! This is not supported. Snippet harvesting process will ignore this line.");
+								Console.WriteLine("Error: Nested OpeningTag found in " + FileName + " at line " + LineNumber + "! This is not supported. Snippet harvesting process will fail.");
 								return false;
 							}
 							else if (TrimmedLine.StartsWith(ClosingTag))
 							{
 								//We're done with this snippet now. Mark that we can end cleanly.
 								IsSnippetBeingProcessed = false;
+								foreach (string CurrentSnippetPageName in CurrentSnippetPageNames)
+								{
+									Snippets.FinishCurrentSnippet(CurrentSnippetPageName);
+								}
 								break;
 							}
 							else if (TrimmedLine.Contains(ClosingTag))
 							{
 								//We will be done this snippet once we add this line (minus the tag).
 								FinishSnippetAfterThisLine = true;
-								CurrentLine.Replace(ClosingTag, "//");
+								CurrentLine = CurrentLine.Replace(ClosingTag, "//");
+							}
+							if (TrimmedLine.Contains(SeeTag))
+							{
+								int FirstCharacter = CurrentLine.IndexOf(SeeTag);
+								if ((FirstCharacter >= 0) && ((FirstCharacter + SeeTag.Length) < CurrentLine.Length))
+								{
+									FirstCharacter += SeeTag.Length;
+
+									SeePageNames = CurrentLine.Substring(FirstCharacter).Split(WhiteSpace).ToList<string>();		//Whitespace is used to delimit our API/snippet page names here.
+									SeePageNames.RemoveAll(entry => (entry.Length < 1));					//Blank entries can show up in the list. Remove them.
+
+									foreach (string CurrentSnippetPageName in CurrentSnippetPageNames)
+									{
+										foreach (string SeeText in SeePageNames)
+										{
+											if (!Snippets.AddSeeText(CurrentSnippetPageName, "1. [](" + SeeText + ")" + Environment.NewLine))
+											{
+												Console.WriteLine("Error: Failed to add text to snippet's \"see\" portion for " + CurrentSnippetPageName + " from source file " + FileName + " at line " + LineNumber);
+												return false;
+											}
+										}
+									}
+									CurrentLine = CurrentLine.Substring(0, FirstCharacter - SeeTag.Length);
+									if (CurrentLine.Trim().Length < 1)
+									{
+										//This line contained only whitespace and this See tag, so skip it entirely.
+										continue;
+									}
+								}
 							}
 
 							if (CurrentLine.Trim().Length < 1)
@@ -843,6 +924,10 @@ namespace APIDocTool
 							if (FinishSnippetAfterThisLine)
 							{
 								IsSnippetBeingProcessed = false;
+								foreach (string CurrentSnippetPageName in CurrentSnippetPageNames)
+								{
+									Snippets.FinishCurrentSnippet(CurrentSnippetPageName);
+								}
 								break;
 							}
 						}
