Improved API Snippet "See Also" text by adding proper page descriptions into See Also links.

diff --git a/Engine/Source/Programs/UnrealDocTool/APIDocTool/APIDocTool/Program.cs b/Engine/Source/Programs/UnrealDocTool/APIDocTool/APIDocTool/Program.cs
--- a/Engine/Source/Programs/UnrealDocTool/APIDocTool/APIDocTool/Program.cs
+++ b/Engine/Source/Programs/UnrealDocTool/APIDocTool/APIDocTool/Program.cs
@@ -879,9 +879,9 @@ namespace APIDocTool
 									{
 										foreach (string SeeText in SeePageNames)
 										{
-											if (!Snippets.AddSeeText(CurrentSnippetPageName, "1. [](" + SeeText + ")" + Environment.NewLine))
+											if (!Snippets.AddSeeText(CurrentSnippetPageName, "1. [](" + SeeText + ")" + " - %" + SeeText + ":description%" + Environment.NewLine))
 											{
-												Console.WriteLine("Error: Failed to add text to snippet's \"see\" portion for " + CurrentSnippetPageName + " from source file " + FileName + " at line " + LineNumber);
+												Console.WriteLine("Error: Failed to add text to snippet's \"See Also\" portion for " + CurrentSnippetPageName + " from source file " + FileName + " at line " + LineNumber);
 												return false;
 											}
 										}
