[REDACTED_HANDLE] tags improved in APIDocTool. Might want to refactor some of this code soon, there's getting to be a pretty big and complex function that does a almost all of the snippet work.

diff --git a/Engine/Source/Programs/UnrealDocTool/APIDocTool/APIDocTool/Program.cs b/Engine/Source/Programs/UnrealDocTool/APIDocTool/APIDocTool/Program.cs
--- a/Engine/Source/Programs/UnrealDocTool/APIDocTool/APIDocTool/Program.cs
+++ b/Engine/Source/Programs/UnrealDocTool/APIDocTool/APIDocTool/Program.cs
@@ -790,12 +790,13 @@ namespace APIDocTool
 			{
 				const string OpeningTag = "///CODE_SNIPPET_START:";
 				const string ClosingTag = "///CODE_SNIPPET_END";
-				const string SeeTag = "///CODE_SNIPPET_SEE:";
+				const string SeeTag = "@see";
 				APISnippets Snippets = new APISnippets();
 				List<string> CurrentSnippetPageNames = new List<string>(4);		//Probably won't have a snippet that is shared by more than four different pages.
 				List<string> SeePageNames = new List<string>(4);				//More than four of these on one line will probably not happen.
 				string CurrentLine;
 				char[] WhiteSpace = {' ', '\t'};		//Doubles as our token delimiter in one place - noted in comments.
+				string[] ClassMemberDelimiters = new string[2] { "::", "()" };
 				bool IsSnippetBeingProcessed = false;
 				bool WasPreviousLineBlank = false;		//Two blank lines in a row will end a code block. Don't allow this to happen in a single snippet.
 
@@ -812,8 +813,7 @@ namespace APIDocTool
 						{
 							continue;
 						}
-						CurrentSnippetPageNames = CurrentLine.Split(WhiteSpace).ToList<string>();		//Whitespace is used to delimit our API/snippet page names here.
-						CurrentSnippetPageNames.RemoveAll(entry => (entry.Length < 1));					//Blank entries can show up in the list. Remove them.
+						CurrentSnippetPageNames = CurrentLine.Split(WhiteSpace, StringSplitOptions.RemoveEmptyEntries).ToList<string>();		//Whitespace is used to delimit our API/snippet page names here.
 						while ((CurrentSnippetPageNames.Count > 0) && (!CurrentSnippetPageNames[0].Contains(OpeningTag)))
 						{
 							CurrentSnippetPageNames.RemoveAt(0);											//Remove everything before the opening tag.
@@ -843,6 +843,7 @@ namespace APIDocTool
 						{
 							++LineNumber;
 							string TrimmedLine = CurrentLine.TrimStart(WhiteSpace);		//This is actually a C# same-line whitespace check, not our token delimiters.
+							string TrimmedLineLower = TrimmedLine.ToLower();
 							if (TrimmedLine.Contains(OpeningTag))
 							{
 								//Snippets do not currently support overlapping. If they did, closing tags should be explicit about which entries are ending, and we'd need to skip lines with opening tags.
@@ -865,32 +866,34 @@ namespace APIDocTool
 								FinishSnippetAfterThisLine = true;
 								CurrentLine = CurrentLine.Replace(ClosingTag, "//");
 							}
-							if (TrimmedLine.Contains(SeeTag))
+							for (int TrimmedLineStartingIndex = TrimmedLineLower.IndexOf(SeeTag); TrimmedLineStartingIndex >= 0; TrimmedLineStartingIndex = TrimmedLineLower.Substring(TrimmedLineStartingIndex + SeeTag.Length).Contains(SeeTag) ? (TrimmedLineStartingIndex + SeeTag.Length + TrimmedLineLower.Substring(TrimmedLineStartingIndex + SeeTag.Length).IndexOf(SeeTag)) : -2)
 							{
-								int FirstCharacter = CurrentLine.IndexOf(SeeTag);
-								if ((FirstCharacter >= 0) && ((FirstCharacter + SeeTag.Length) < CurrentLine.Length))
-								{
-									FirstCharacter += SeeTag.Length;
-
-									SeePageNames = CurrentLine.Substring(FirstCharacter).Split(WhiteSpace).ToList<string>();		//Whitespace is used to delimit our API/snippet page names here.
-									SeePageNames.RemoveAll(entry => (entry.Length < 1));					//Blank entries can show up in the list. Remove them.
+								int FirstCharacter = TrimmedLineStartingIndex + SeeTag.Length;
+								SeePageNames = TrimmedLine.Substring(FirstCharacter).Split(WhiteSpace, StringSplitOptions.RemoveEmptyEntries).ToList<string>();		//Whitespace is used to delimit our API/snippet page names here.
 
+								if (SeePageNames.Count > 0)
+								{
+									string SeeText = SeePageNames[0];
+									string[] SeeTextBreakdown = SeeText.Split(ClassMemberDelimiters, StringSplitOptions.RemoveEmptyEntries);
 									foreach (string CurrentSnippetPageName in CurrentSnippetPageNames)
 									{
-										foreach (string SeeText in SeePageNames)
+										if (SeeTextBreakdown.Length == 2)
 										{
-											if (!Snippets.AddSeeText(CurrentSnippetPageName, "1. [](" + SeeText + ")" + " - %" + SeeText + ":description%" + Environment.NewLine))
+											//This is considered "class::member" format.
+											if (!Snippets.AddSeeText(CurrentSnippetPageName, "1. [](API:" + SeeText + ")" + Environment.NewLine))
+											{
+												Console.WriteLine("Error: Failed to add text to snippet's \"See Also\" portion for " + CurrentSnippetPageName + " from source file " + FileName + " at line " + LineNumber);
+												return false;
+											}
+										}
+										else
+										{
+											if (!Snippets.AddSeeText(CurrentSnippetPageName, "1. [](" + SeeText + ")" + Environment.NewLine))
 											{
 												Console.WriteLine("Error: Failed to add text to snippet's \"See Also\" portion for " + CurrentSnippetPageName + " from source file " + FileName + " at line " + LineNumber);
 												return false;
 											}
 										}
-									}
-									CurrentLine = CurrentLine.Substring(0, FirstCharacter - SeeTag.Length);
-									if (CurrentLine.Trim().Length < 1)
-									{
-										//This line contained only whitespace and this See tag, so skip it entirely.
-										continue;
 									}
 								}
 							}
