- Fixed Templates folder being ignored
 - Added Skill Level metadata
 - Added SEO-Title metadata
 - Added SEO-Description metadata

diff --git a/Engine/Source/Programs/UnrealDocTool/UnrealDocTool/MarkdownSharp/Markdown.cs b/Engine/Source/Programs/UnrealDocTool/UnrealDocTool/MarkdownSharp/Markdown.cs
--- a/Engine/Source/Programs/UnrealDocTool/UnrealDocTool/MarkdownSharp/Markdown.cs
+++ b/Engine/Source/Programs/UnrealDocTool/UnrealDocTool/MarkdownSharp/Markdown.cs
@@ -1060,13 +1060,31 @@ namespace MarkdownSharp
                         break;
                 }
 
-                metadata = metadataRow.Value.Aggregate(metadata,
-                    (current, metadataValue) =>
-                    current + Templates.Metadata.Render(Hash.FromAnonymousObject(new
+                if (metadataRow.Key == "description" && data.ProcessedDocumentCache.Metadata.MetadataMap.ContainsKey("seo-description"))
+                {
+                    continue;
+                }
+                
+                if (metadataRow.Key == "seo-description")
+                {
+                    metadata = metadataRow.Value.Aggregate(metadata,
+                        (current, metadataValue) =>
+                        current + Templates.Metadata.Render(Hash.FromAnonymousObject(new
                         {
-                            key = metadataRow.Key,
+                            key = "description",
                             value = metadataValue
                         })));
+                }
+                else
+                {
+                    metadata = metadataRow.Value.Aggregate(metadata,
+                        (current, metadataValue) =>
+                        current + Templates.Metadata.Render(Hash.FromAnonymousObject(new
+                            {
+                                key = metadataRow.Key,
+                                value = metadataValue
+                            })));
+                }
             }
 
             var translatedPageLinks = "";
@@ -1128,6 +1146,7 @@ namespace MarkdownSharp
                                          ? ""
                                          : Templates.GetCached("crumbsDiv.html").Render(Hash.FromAnonymousObject(new { crumbs })),
                         title = data.ProcessedDocumentCache.CurrentFileDocument.Metadata.DocumentTitle,
+                        seotitle = (data.ProcessedDocumentCache.CurrentFileDocument.Metadata.SEOTitle != null) ? data.ProcessedDocumentCache.CurrentFileDocument.Metadata.SEOTitle : data.ProcessedDocumentCache.CurrentFileDocument.Metadata.DocumentTitle,
                         metadata = metadata,
                         translatedPages = translatedPageLinks,
                         relatedPagesMenu = Templates.RelatedPages.Render(Hash.FromAnonymousObject(
@@ -1156,6 +1175,12 @@ namespace MarkdownSharp
                                 versions = data.ProcessedDocumentCache.Metadata.EngineVersions,
                                 versionCount = data.ProcessedDocumentCache.Metadata.EngineVersions.Count
                             })),
+                        skilllevels = Templates.SkillLevels.Render(Hash.FromAnonymousObject(
+                            new
+                            {
+                                skilllevels = data.ProcessedDocumentCache.Metadata.SkillLevels,
+                                skilllevelCount = data.ProcessedDocumentCache.Metadata.SkillLevels.Count
+                            })),
                         errors = ThisIsPreview
                                      ? Templates.ErrorDetails.Render(Hash.FromAnonymousObject(
                                          new
diff --git a/Engine/Source/Programs/UnrealDocTool/UnrealDocTool/MarkdownSharp/Preprocessor/MetadataManager.cs b/Engine/Source/Programs/UnrealDocTool/UnrealDocTool/MarkdownSharp/Preprocessor/MetadataManager.cs
--- a/Engine/Source/Programs/UnrealDocTool/UnrealDocTool/MarkdownSharp/Preprocessor/MetadataManager.cs
+++ b/Engine/Source/Programs/UnrealDocTool/UnrealDocTool/MarkdownSharp/Preprocessor/MetadataManager.cs
@@ -64,11 +64,16 @@ namespace MarkdownSharp.Preprocessor
                     DocumentTitle = metaValue;
                 }
 
-                if (metaDataCategoryLowerCase == "crumbs")
+                if (metaDataCategoryLowerCase == "seo-title")
                 {
-                    CrumbsLinks.Add(metaValue);
+                    SEOTitle = metaValue;
                 }
 
+                //if (metaDataCategoryLowerCase == "crumbs")
+                //{
+                //    CrumbsLinks.Add(metaValue);
+                //}
+
                 if (metaDataCategoryLowerCase == "related" && full)
                 {
                     RelatedLinks.Add(data.Markdown.ProcessRelated(metaValue, data));
@@ -83,10 +88,16 @@ namespace MarkdownSharp.Preprocessor
                 {
                     EngineVersions.Add(metaValue);
                 }
+
+                if (metaDataCategoryLowerCase == "skilllevel")
+                {
+                    SkillLevels.Add(metaValue);
+                }
             }
 
             // Add meta data to the list, we require some specific meta data keys to be unique others can be duplicates
-            if (metaDataCategoryLowerCase.Equals("title") || metaDataCategoryLowerCase.Equals("description")
+            if (metaDataCategoryLowerCase.Equals("title") || metaDataCategoryLowerCase.Equals("seo-title")
+                || metaDataCategoryLowerCase.Equals("description") || metaDataCategoryLowerCase.Equals("seo-description")
                 || metaDataCategoryLowerCase.Equals("template") || metaDataCategoryLowerCase.Equals("forcepublishfiles"))
             {
                 if (MetadataMap.ContainsKey(metaDataCategoryLowerCase))
@@ -145,6 +156,7 @@ namespace MarkdownSharp.Preprocessor
             RelatedLinks = new List<Hash>();
             PrereqLinks = new List<Hash>();
             EngineVersions = new List<String>();
+            SkillLevels = new List<String>();
 
             var changes = new List<PreprocessingTextChange>();
 
@@ -207,7 +219,9 @@ namespace MarkdownSharp.Preprocessor
         public List<Hash> RelatedLinks { get; set; }
         public List<Hash> PrereqLinks { get; set; }
         public List<String> EngineVersions { get; set; }
+        public List<String> SkillLevels { get; set; }
         public string DocumentTitle { get; set; }
+        public string SEOTitle { get; set; }
 
         public Dictionary<string, List<string>> MetadataMap { get; private set; }
     }
diff --git a/Engine/Source/Programs/UnrealDocTool/UnrealDocTool/MarkdownSharp/Templates.cs b/Engine/Source/Programs/UnrealDocTool/UnrealDocTool/MarkdownSharp/Templates.cs
--- a/Engine/Source/Programs/UnrealDocTool/UnrealDocTool/MarkdownSharp/Templates.cs
+++ b/Engine/Source/Programs/UnrealDocTool/UnrealDocTool/MarkdownSharp/Templates.cs
@@ -355,6 +355,11 @@ namespace MarkdownSharp
             get { return GetCached("Version.html"); }
         }
 
+        public static TemplateFile SkillLevels
+        {
+            get { return GetCached("SkillLevel.html"); }
+        }
+
         public static TemplateFile MissingTemplate
         {
             get { return new TemplateFile(); }
diff --git a/Engine/Source/Programs/UnrealDocTool/UnrealDocTool/UnrealDocTool/Program.cs b/Engine/Source/Programs/UnrealDocTool/UnrealDocTool/UnrealDocTool/Program.cs
--- a/Engine/Source/Programs/UnrealDocTool/UnrealDocTool/UnrealDocTool/Program.cs
+++ b/Engine/Source/Programs/UnrealDocTool/UnrealDocTool/UnrealDocTool/Program.cs
@@ -1053,7 +1053,7 @@ namespace UnrealDocTool
 
         static void CleanMetaDataRecursiveDirectory(string sourcePath)
         {
-            DoRecursively(sourcePath, recursivePath => CleanMetaDataDirectory(recursivePath));
+            DoRecursivelyIgnored(sourcePath, recursivePath => CleanMetaDataDirectory(recursivePath));
         }
 
         private static readonly string[] IgnoredFolders = new[] { "include", "javascript", "images", "attachments", "css", "templates" };
@@ -1081,11 +1081,29 @@ namespace UnrealDocTool
 
 			action(sourcePath);
 
+            foreach (var subDirectory in Directory.GetDirectories(sourcePath))
+            {
+                //if (!IsIgnoredDirectory(subDirectory))
+                //{
+                    DoRecursively(subDirectory, action);
+                //}
+            }
+        }
+
+        private static void DoRecursivelyIgnored(string sourcePath, Action<string> action)
+        {
+            if (!Directory.Exists(sourcePath))
+            {
+                return;
+            }
+
+            action(sourcePath);
+
             foreach (var subDirectory in Directory.GetDirectories(sourcePath))
             {
                 if (!IsIgnoredDirectory(subDirectory))
                 {
-                    DoRecursively(subDirectory, action);
+                    DoRecursivelyIgnored(subDirectory, action);
                 }
             }
         }
diff --git a/Engine/Source/Programs/UnrealDocTool/UnrealDocTool/UnrealDocTool/app.config b/Engine/Source/Programs/UnrealDocTool/UnrealDocTool/UnrealDocTool/app.config
--- a/Engine/Source/Programs/UnrealDocTool/UnrealDocTool/UnrealDocTool/app.config
+++ b/Engine/Source/Programs/UnrealDocTool/UnrealDocTool/UnrealDocTool/app.config
@@ -52,7 +52,7 @@
         <value>Default.html</value>
       </setting>
       <setting name="MetadataErrorIfMissing" serializeAs="String">
-        <value>Title,Crumbs,Availability</value>
+        <value>Title,Availability</value>
       </setting>
       <setting name="MetadataInfoIfMissing" serializeAs="String">
         <value>Description</value>
