Added crumbs parsing back

diff --git a/Engine/Source/Programs/UnrealDocTool/UnrealDocTool/MarkdownSharp/Preprocessor/MetadataManager.cs b/Engine/Source/Programs/UnrealDocTool/UnrealDocTool/MarkdownSharp/Preprocessor/MetadataManager.cs
--- a/Engine/Source/Programs/UnrealDocTool/UnrealDocTool/MarkdownSharp/Preprocessor/MetadataManager.cs
+++ b/Engine/Source/Programs/UnrealDocTool/UnrealDocTool/MarkdownSharp/Preprocessor/MetadataManager.cs
@@ -69,10 +69,10 @@ namespace MarkdownSharp.Preprocessor
                     SEOTitle = metaValue;
                 }
 
-                //if (metaDataCategoryLowerCase == "crumbs")
-                //{
-                //    CrumbsLinks.Add(metaValue);
-                //}
+                if (metaDataCategoryLowerCase == "crumbs")
+                {
+                    CrumbsLinks.Add(metaValue);
+                }
 
                 if (metaDataCategoryLowerCase == "related" && full)
                 {
