Fix APIDocTool to account for potential removal of crumbs from UnrealDocTool.

diff --git a/Engine/Source/Programs/UnrealDocTool/APIDocTool/APIDocTool/Program.cs b/Engine/Source/Programs/UnrealDocTool/APIDocTool/APIDocTool/Program.cs
--- a/Engine/Source/Programs/UnrealDocTool/APIDocTool/APIDocTool/Program.cs
+++ b/Engine/Source/Programs/UnrealDocTool/APIDocTool/APIDocTool/Program.cs
@@ -1559,9 +1559,12 @@ namespace APIDocTool
 						HtmlText = HtmlText.Remove(HeaderMinIdx, HeaderMaxIdx + HeaderEndText.Length - HeaderMinIdx);
 
 						int CrumbsMinIdx = HtmlText.IndexOf("<div class=\"crumbs\">");
-						int HomeMinIdx = HtmlText.IndexOf("<strong>", CrumbsMinIdx);
-						int HomeMaxIdx = HtmlText.IndexOf("&gt;", HomeMinIdx) + 4;
-						HtmlText = HtmlText.Remove(HomeMinIdx, HomeMaxIdx - HomeMinIdx);
+						if (CrumbsMinIdx >= 0)
+						{
+							int HomeMinIdx = HtmlText.IndexOf("<strong>", CrumbsMinIdx);
+							int HomeMaxIdx = HtmlText.IndexOf("&gt;", HomeMinIdx) + 4;
+							HtmlText = HtmlText.Remove(HomeMinIdx, HomeMaxIdx - HomeMinIdx);
+						}
 
 						File.WriteAllText(TargetFileName, HtmlText);
 					}
