[UEDOC-1625] Remove space before * or & in function parameter lists generated by APIDocTool to match Epic code standards.

diff --git a/Engine/Source/Programs/UnrealDocTool/APIDocTool/APIDocTool/APIFunction.cs b/Engine/Source/Programs/UnrealDocTool/APIDocTool/APIDocTool/APIFunction.cs
--- a/Engine/Source/Programs/UnrealDocTool/APIDocTool/APIDocTool/APIFunction.cs
+++ b/Engine/Source/Programs/UnrealDocTool/APIDocTool/APIDocTool/APIFunction.cs
@@ -717,6 +717,11 @@ namespace APIDocTool
 					{
 						string Separator = (Idx + 1 == Parameters.Count) ? "" : ",";
 						string Definition = Markdown.Truncate(APIMember.RemoveElaborations(Parameters[Idx].Definition), 35, "...");
+						// Fix spacing around pointer/reference punctuation to match Epic code standards.
+						if (Definition != null)
+						{
+							Definition = Definition.Replace(" * ", "* ").Replace(" & ", "& ");
+						}
 						Writer.WriteLine(UdnWriter.TabSpaces + Definition + Separator + "  ");
 					}
 					Writer.WriteEscapedLine(")  ");
