APIDocTool temp hack to remove duplicate output delegate pins. Will be removed when UE-16475 is resolved.

diff --git a/Engine/Source/Programs/UnrealDocTool/APIDocTool/APIDocTool/APIAction.cs b/Engine/Source/Programs/UnrealDocTool/APIDocTool/APIDocTool/APIAction.cs
--- a/Engine/Source/Programs/UnrealDocTool/APIDocTool/APIDocTool/APIAction.cs
+++ b/Engine/Source/Programs/UnrealDocTool/APIDocTool/APIDocTool/APIAction.cs
@@ -201,7 +201,25 @@ namespace APIDocTool
 					// Outputs
 					Writer.EnterSection("outputs", "Outputs");
 					Writer.WriteObject("MemberIconListHeadBlank");
-					WritePins(Writer, Pins.Where(x => !x.bInputPin));
+					// TODO: Remove this hack and reinstate the one-line version once UE-16475 is resolved.
+					bool bAlreadyWroteOutputDelegate = false;
+					for (int i = 0; i < Pins.Count; ++i)
+					{
+						APIActionPin Pin = Pins[i];
+						if (!Pin.bInputPin)
+						{
+							if (Pin.GetTypeText() == "delegate")
+							{
+								if (bAlreadyWroteOutputDelegate)
+								{
+									continue;
+								}
+								bAlreadyWroteOutputDelegate = true;
+							}
+							Pin.WritePin(Writer);
+						}
+					}
+					//WritePins(Writer, Pins.Where(x => !x.bInputPin));
 					Writer.WriteObject("MemberIconListTail");
 					Writer.LeaveSection();
 				}
