#UE4 Docs: added framework for sitemap filters (tags, skill level filters, engine version filters)

diff --git a/Engine/Source/Programs/UnrealDocTool/UnrealDocTool/MarkdownSharp/Markdown.cs b/Engine/Source/Programs/UnrealDocTool/UnrealDocTool/MarkdownSharp/Markdown.cs
--- a/Engine/Source/Programs/UnrealDocTool/UnrealDocTool/MarkdownSharp/Markdown.cs
+++ b/Engine/Source/Programs/UnrealDocTool/UnrealDocTool/MarkdownSharp/Markdown.cs
@@ -1073,7 +1073,7 @@ namespace MarkdownSharp
                         {
                             key = "description",
                             value = metadataValue
-                        })));
+                        })) + Environment.NewLine);
                 }
                 else
                 {
@@ -1083,7 +1083,7 @@ namespace MarkdownSharp
                             {
                                 key = metadataRow.Key,
                                 value = metadataValue
-                            })));
+                            })) + Environment.NewLine);
                 }
             }
 
@@ -1173,13 +1173,25 @@ namespace MarkdownSharp
                             new
                             {
                                 versions = data.ProcessedDocumentCache.Metadata.EngineVersions,
-                                versionCount = data.ProcessedDocumentCache.Metadata.EngineVersions.Count
+                                versionCount = data.ProcessedDocumentCache.Metadata.EngineVersions.Count,
+                                language = data.CurrentFolderDetails.Language,
+                                relativeHtmlPath = data.CurrentFolderDetails.RelativeHTMLPath
                             })),
                         skilllevels = Templates.SkillLevels.Render(Hash.FromAnonymousObject(
                             new
                             {
                                 skilllevels = data.ProcessedDocumentCache.Metadata.SkillLevels,
-                                skilllevelCount = data.ProcessedDocumentCache.Metadata.SkillLevels.Count
+                                skilllevelCount = data.ProcessedDocumentCache.Metadata.SkillLevels.Count,
+                                language = data.CurrentFolderDetails.Language,
+                                relativeHtmlPath = data.CurrentFolderDetails.RelativeHTMLPath
+                            })),
+                        tags = Templates.Tags.Render(Hash.FromAnonymousObject(
+                            new
+                            {
+                                tags = data.ProcessedDocumentCache.Metadata.Tags,
+                                tagsCount = data.ProcessedDocumentCache.Metadata.Tags.Count,
+                                language = data.CurrentFolderDetails.Language,
+                                relativeHtmlPath = data.CurrentFolderDetails.RelativeHTMLPath
                             })),
                         errors = ThisIsPreview
                                      ? Templates.ErrorDetails.Render(Hash.FromAnonymousObject(
diff --git a/Engine/Source/Programs/UnrealDocTool/UnrealDocTool/MarkdownSharp/Preprocessor/MetadataManager.cs b/Engine/Source/Programs/UnrealDocTool/UnrealDocTool/MarkdownSharp/Preprocessor/MetadataManager.cs
--- a/Engine/Source/Programs/UnrealDocTool/UnrealDocTool/MarkdownSharp/Preprocessor/MetadataManager.cs
+++ b/Engine/Source/Programs/UnrealDocTool/UnrealDocTool/MarkdownSharp/Preprocessor/MetadataManager.cs
@@ -5,9 +5,25 @@ using System.Collections.Generic;
 using System.Linq;
 using System.Text.RegularExpressions;
 using DotLiquid;
+using System.Globalization;
 
 namespace MarkdownSharp.Preprocessor
 {
+    public class EngineVersion
+    {
+        public String VersionNumber = "";
+        public String DisplayLabel = "";
+
+        public EngineVersion(String version, String label)
+        {
+            VersionNumber = version;
+            DisplayLabel = label;
+        }
+    }
+    public class EngineVersionDrop : Drop
+    {
+
+    }
     public class MetadataManager
     {
         public PreprocessedData Data { get; private set; }
@@ -86,13 +102,18 @@ namespace MarkdownSharp.Preprocessor
 
                 if (metaDataCategoryLowerCase == "version")
                 {
-                    EngineVersions.Add(metaValue);
+                    EngineVersions.Add(Hash.FromAnonymousObject(new { version = metaValue, label = metaValue.Replace('.', '_') }));
                 }
 
                 if (metaDataCategoryLowerCase == "skilllevel")
                 {
                     SkillLevels.Add(metaValue);
                 }
+
+                if (metaDataCategoryLowerCase == "tags")
+                {
+                    Tags.Add(CultureInfo.CurrentCulture.TextInfo.ToTitleCase(metaValue));
+                }
             }
 
             // Add meta data to the list, we require some specific meta data keys to be unique others can be duplicates
@@ -155,8 +176,9 @@ namespace MarkdownSharp.Preprocessor
             CrumbsLinks = new List<string>();
             RelatedLinks = new List<Hash>();
             PrereqLinks = new List<Hash>();
-            EngineVersions = new List<String>();
+            EngineVersions = new List<Hash>();
             SkillLevels = new List<String>();
+            Tags = new List<String>();
 
             var changes = new List<PreprocessingTextChange>();
 
@@ -218,8 +240,9 @@ namespace MarkdownSharp.Preprocessor
         public List<string> CrumbsLinks { get; set; }
         public List<Hash> RelatedLinks { get; set; }
         public List<Hash> PrereqLinks { get; set; }
-        public List<String> EngineVersions { get; set; }
+        public List<Hash> EngineVersions { get; set; }
         public List<String> SkillLevels { get; set; }
+        public List<String> Tags { get; set; }
         public string DocumentTitle { get; set; }
         public string SEOTitle { get; set; }
 
diff --git a/Engine/Source/Programs/UnrealDocTool/UnrealDocTool/MarkdownSharp/Templates.cs b/Engine/Source/Programs/UnrealDocTool/UnrealDocTool/MarkdownSharp/Templates.cs
--- a/Engine/Source/Programs/UnrealDocTool/UnrealDocTool/MarkdownSharp/Templates.cs
+++ b/Engine/Source/Programs/UnrealDocTool/UnrealDocTool/MarkdownSharp/Templates.cs
@@ -360,6 +360,11 @@ namespace MarkdownSharp
             get { return GetCached("SkillLevel.html"); }
         }
 
+        public static TemplateFile Tags
+        {
+            get { return GetCached("Tag.html"); }
+        }
+
         public static TemplateFile MissingTemplate
         {
             get { return new TemplateFile(); }
