[UE-17095] Added a python script to handle non-Doxygen-format comments.

diff --git a/Engine/Source/Programs/UnrealDocTool/DoxygenInputFilter/CommentFixup.py b/Engine/Source/Programs/UnrealDocTool/DoxygenInputFilter/CommentFixup.py
new file mode 100644
--- /dev/null
+++ b/Engine/Source/Programs/UnrealDocTool/DoxygenInputFilter/CommentFixup.py
@@ -0,0 +1,11 @@
+#!/usr/bin/env python
+import sys
+f = file(sys.argv[1], 'r')
+line = f.readline()
+while line:
+    line = line + " "
+    line = line.replace("/* ", "/** ")
+    line = line.replace("// ", "/// ")
+    sys.stdout.write(line)
+    line = f.readline()
+f.close() 
\ No newline at end of file
diff --git a/Engine/Source/Programs/UnrealDocTool/DoxygenLib/Doxygen.cs b/Engine/Source/Programs/UnrealDocTool/DoxygenLib/Doxygen.cs
--- a/Engine/Source/Programs/UnrealDocTool/DoxygenLib/Doxygen.cs
+++ b/Engine/Source/Programs/UnrealDocTool/DoxygenLib/Doxygen.cs
@@ -30,7 +30,7 @@ namespace DoxygenLib
 			OutputPath = InOutputPath;
 		}
 
-		public void Write(string ConfigPath)
+		public void Write(string ConfigPath, string ExecutePath)
 		{
 			using (StreamWriter Output = new StreamWriter(ConfigPath))
 			{
@@ -68,6 +68,19 @@ namespace DoxygenLib
 				Output.WriteLine();
 				FormatSetting(Output, "INPUT", new List<string>(InputPaths.Select(x => QuoteValue(x))));
 
+				// Input filter program
+				Output.WriteLine();
+				//FormatSetting(Output, "INPUT_FILTER", Path.GetDirectoryName(ExecutePath) + "/../../../../../DoxygenInputFilter/CommentFixup.py");
+				string FilterPath = Path.GetDirectoryName(ExecutePath).Replace("\\", "/");
+				// Back up by five levels, then tack on the relative path to the filter script.
+				for (int PathRemovalCoiunt = 0; PathRemovalCoiunt < 5; ++PathRemovalCoiunt)
+				{
+					FilterPath = FilterPath.Remove(FilterPath.LastIndexOf('/'));
+				}
+				FilterPath += "/Source/Programs/UnrealDocTool/DoxygenInputFilter/CommentFixup.py";
+				FormatSetting(Output, "INPUT_FILTER", FilterPath);
+				//FormatSetting(Output, "INPUT_FILTER", "D:/Builds_ArtistSync/UE4/Engine/Source/Programs/UnrealDocTool/DoxygenInputFilter/CommentFixup.py");
+
 				// List of predefined macros
 				Output.WriteLine();
 				FormatSetting(Output, "PREDEFINED", new List<string>(Definitions.Select(x => QuoteValue(x))));
@@ -153,7 +166,7 @@ namespace DoxygenLib
 
 			// Write the doxygen config file
 			string ConfigFilePath = Path.Combine(Config.OutputPath, "Doxyfile");
-			Config.Write(ConfigFilePath);
+			Config.Write(ConfigFilePath, DoxygenPath);
 
 			// Spawn doxygen
 			using (Process DoxygenProcess = new Process())
