Minor python script cleanup.

diff --git a/Engine/Source/Programs/UnrealDocTool/DoxygenInputFilter/CommentFixup.py b/Engine/Source/Programs/UnrealDocTool/DoxygenInputFilter/CommentFixup.py
--- a/Engine/Source/Programs/UnrealDocTool/DoxygenInputFilter/CommentFixup.py
+++ b/Engine/Source/Programs/UnrealDocTool/DoxygenInputFilter/CommentFixup.py
@@ -1,8 +1,8 @@
-#!/usr/bin/env python
 import sys
 f = file(sys.argv[1], 'r')
 line = f.readline()
 while line:
+    # Pad the string with an extra space at the end to catch end-of-line block comments.
     line = line + " "
     line = line.replace("/* ", "/** ")
     line = line.replace("// ", "/// ")
