- fixed relative paths in markdown plugin

diff --git a/Engine/Source/Programs/UnrealDocTool/UnrealDocTool/MarkdownSharp/EpicMarkdown/EMDocument.cs b/Engine/Source/Programs/UnrealDocTool/UnrealDocTool/MarkdownSharp/EpicMarkdown/EMDocument.cs
--- a/Engine/Source/Programs/UnrealDocTool/UnrealDocTool/MarkdownSharp/EpicMarkdown/EMDocument.cs
+++ b/Engine/Source/Programs/UnrealDocTool/UnrealDocTool/MarkdownSharp/EpicMarkdown/EMDocument.cs
@@ -159,7 +159,7 @@ namespace MarkdownSharp.EpicMarkdown
         public string GetRelativeTargetPath(EMDocument current)
         {
             var relPath = FileHelper.GetRelativePath(Path.GetDirectoryName(current.LocalPath), Path.GetDirectoryName(LocalPath));
-            return string.IsNullOrWhiteSpace(relPath) ? "" : Path.Combine(relPath, "index.html");
+            return string.IsNullOrWhiteSpace(relPath) ? "" : Path.Combine(relPath, "index.html").Replace('\\', '/');
         }
 
         public List<string> GetUsedImagesPaths()
@@ -249,7 +249,7 @@ namespace MarkdownSharp.EpicMarkdown
                 var path = link.Path as EMLocalDocumentPath;
                 var id = path.BookmarkName;
 
-                if ((path.IsBookmark && !bookmarks.ContainsKey(id))
+                if ((path.IsBookmark && (path.DocumentPath != path.currentDocument.LocalPath || !bookmarks.ContainsKey(id)))
                     || (!path.IsBookmark && !inclusions.ContainsKey(path.DocumentPath.ToLower())))
                 {
                     if (data.NonDynamicHTMLOutput)
diff --git a/Engine/Source/Programs/UnrealDocTool/UnrealDocTool/MarkdownSharp/EpicMarkdown/PathProviders/EMLocalDocumentPath.cs b/Engine/Source/Programs/UnrealDocTool/UnrealDocTool/MarkdownSharp/EpicMarkdown/PathProviders/EMLocalDocumentPath.cs
--- a/Engine/Source/Programs/UnrealDocTool/UnrealDocTool/MarkdownSharp/EpicMarkdown/PathProviders/EMLocalDocumentPath.cs
+++ b/Engine/Source/Programs/UnrealDocTool/UnrealDocTool/MarkdownSharp/EpicMarkdown/PathProviders/EMLocalDocumentPath.cs
@@ -12,7 +12,7 @@ namespace MarkdownSharp.EpicMarkdown.PathProviders
 
     public class EMLocalDocumentPath : EMPathProvider
     {
-        private readonly EMDocument currentDocument;
+        public readonly EMDocument currentDocument;
 
         private static readonly Regex DocumentPathPattern = new Regex("^(?<path>[^#]*)(#(?<bookmark>[^#]+))?$", RegexOptions.Compiled);
 
