tables without headers no longer make the first column use TH elements

diff --git a/Engine/Source/Programs/UnrealDocTool/UnrealDocTool/MarkdownSharp/EpicMarkdown/EMTable.cs b/Engine/Source/Programs/UnrealDocTool/UnrealDocTool/MarkdownSharp/EpicMarkdown/EMTable.cs
--- a/Engine/Source/Programs/UnrealDocTool/UnrealDocTool/MarkdownSharp/EpicMarkdown/EMTable.cs
+++ b/Engine/Source/Programs/UnrealDocTool/UnrealDocTool/MarkdownSharp/EpicMarkdown/EMTable.cs
@@ -214,8 +214,7 @@ namespace MarkdownSharp.EpicMarkdown
                     cell = Normalizer.TrailingWhitespaceRemove(cell);
                     
                     var columnSpanLength = 1;
-                    var isRowHeader = (count == 0 && useRowHeader && dataColumns.Count > 1)
-                                      || (dataColumns.Count == 1 && i == 0 && !hasHeader);
+                    var isRowHeader = (dataColumns.Count == 1 && i == 0 && !hasHeader);
 
                     if (Regex.Match(dataColumn.Groups[2].Value, @"(\|{2,})").Success)
                     {
