UnrealDocTranslator: Fix incorrect path sometimes being returned by GetUnrealDocToolLocation()

diff --git a/Engine/Source/Programs/UnrealDocTool/UnrealDocTranslator/UnrealDocTranslator/UnrealDocTranslatorForm.cs b/Engine/Source/Programs/UnrealDocTool/UnrealDocTranslator/UnrealDocTranslator/UnrealDocTranslatorForm.cs
--- a/Engine/Source/Programs/UnrealDocTool/UnrealDocTranslator/UnrealDocTranslator/UnrealDocTranslatorForm.cs
+++ b/Engine/Source/Programs/UnrealDocTool/UnrealDocTranslator/UnrealDocTranslator/UnrealDocTranslatorForm.cs
@@ -687,9 +687,8 @@ namespace UnrealDocTranslator
 
         private string GetUnrealDocToolLocation()
         {
-            return Path.Combine(
-                    P4Repository.Connection.Client.Root,
-                    "UE4", "Engine", "Binaries", "DotNET", "UnrealDocTool.exe");
+			FileInfo file = new FileInfo(System.Reflection.Assembly.GetExecutingAssembly().Location);
+			return Path.Combine(file.DirectoryName, "UnrealDocTool.exe");
         }
 
         private static Thread RunLoggingThread(Process process)
