# Doc Translator Tools update.
- Now we will copy all INT files to other languages automaticly and mark those copied version as untranslated by INTSourceChangelist:0.  This can help fix all those reference errors and also make other languages not fall bak to INT version while link to not translated pages.

diff --git a/Engine/Source/Programs/UnrealDocTool/UnrealDocTranslator/UnrealDocTranslator/UnrealDocTranslatorForm.cs b/Engine/Source/Programs/UnrealDocTool/UnrealDocTranslator/UnrealDocTranslator/UnrealDocTranslatorForm.cs
--- a/Engine/Source/Programs/UnrealDocTool/UnrealDocTranslator/UnrealDocTranslator/UnrealDocTranslatorForm.cs
+++ b/Engine/Source/Programs/UnrealDocTool/UnrealDocTranslator/UnrealDocTranslator/UnrealDocTranslatorForm.cs
@@ -490,6 +490,8 @@ namespace UnrealDocTranslator
         private void SyncP4Button_Click(object sender, EventArgs e)
         {
             RefreshP4Details();
+            // mi.wang: Some magic goes here, we should copy those missing files for languages and mark INTSourceChangelist as 0.
+
         }
 
         /// <summary>
@@ -1677,19 +1679,9 @@ namespace UnrealDocTranslator
                 }
             }
 
-            //Setup list of files for creation
-            log.Info("Generating list of missing language files");
+            //Setup list of files for create, change and translated
+            log.Info("Generating list of language files missing, needing update or already translated");
             LangFileToCreateProcessingList.Clear();
-            foreach (var INTFile in INTFileProcessingDictionary)
-            {
-                if (!LangFileProcessingDictionary.ContainsKey(INTFile.Key))
-                {
-                    LangFileToCreateProcessingList.Add(new FileProcessingDetails(INTFile.Value.DepotFile, INTFile.Value.ClientFile, INTFile.Value.ChangeList, FileProcessingDetails.State.New, INTFile.Value.CheckedOutByOthers, _connectionDetails));
-                }
-            }
-
-            //Setup list of files for change and translated
-            log.Info("Generating list of language files needing update or already translated");
             LangFileToUpdateProcessingList.Clear();
             LangFileTranslatedProcessingList.Clear();
             foreach (var INTFile in INTFileProcessingDictionary)
@@ -1699,7 +1691,13 @@ namespace UnrealDocTranslator
                     FileProcessingDetails LangFile;
                     if (LangFileProcessingDictionary.TryGetValue(INTFile.Key, out LangFile))
                     {
-                        if (LangFile.ChangeList < INTFile.Value.ChangeList)
+                        // mi.wang: Special case we can now copy all those untranslated pages to localized folder but mark ChangeList as 0 to distinguish them from Needs Updated ones.
+                        //          This can help fix reference errors from actual parse errors and also keep links not fall back to INT while using specified language under UDN.
+                        if (LangFile.ChangeList <= 0)
+                        {
+                            LangFileToCreateProcessingList.Add(new FileProcessingDetails(INTFile.Value.DepotFile, INTFile.Value.ClientFile, INTFile.Value.ChangeList, FileProcessingDetails.State.New, INTFile.Value.CheckedOutByOthers, _connectionDetails));
+                        }
+                        else if (LangFile.ChangeList < INTFile.Value.ChangeList)
                         {
                             // Needs update
                             LangFileToUpdateProcessingList.Add(new FileProcessingDetails(INTFile.Value.DepotFile, INTFile.Value.ClientFile, INTFile.Value.ChangeList, FileProcessingDetails.State.Changed, INTFile.Value.CheckedOutByOthers, _connectionDetails));
@@ -1711,6 +1709,38 @@ namespace UnrealDocTranslator
                         }
                     }
                 }
+                else
+                {
+                    LangFileToCreateProcessingList.Add(new FileProcessingDetails(INTFile.Value.DepotFile, INTFile.Value.ClientFile, INTFile.Value.ChangeList, FileProcessingDetails.State.New, INTFile.Value.CheckedOutByOthers, _connectionDetails));
+                    string LangClientFileName =  GetLangFileNameFromInt( INTFile.Value.ClientFile );
+                    if ( !(new System.IO.FileInfo(LangClientFileName)).Exists )
+                    {
+                        //Add Changelist number of INT file to top of lang file as metadata
+                        string FileContents = "INTSourceChangelist:0" + Environment.NewLine +
+                                                System.IO.File.ReadAllText(INTFile.Value.ClientFile, Encoding.UTF8);
+                        System.IO.File.WriteAllText(LangClientFileName, FileContents, Encoding.UTF8);
+                        if (CheckConnect())
+                        {
+                            using (Connection P4Connection = P4Repository.Connection)
+                            {
+                                FileSpec[] CurrentClientFile = new FileSpec[1];
+                                CurrentClientFile[0] = new FileSpec(null, new ClientPath(LangClientFileName), null, null);
+
+                                try
+                                {
+                                    P4Connection.Client.AddFiles(CurrentClientFile, null);
+                                    FilesCheckedOut.Add( GetLangFileNameFromInt(INTFile.Value.DepotFile) );
+                                }
+                                catch (P4CommandTimeOutException Error)
+                                {
+                                    //do not continue
+                                    log.Error(Error.Message);
+                                    return;
+                                }
+                            }
+                        }
+                    }
+                }
             }
 
             //Setup list of files for move
