Unreal Doc Tool Updates

 - Fancy Links
 - Language selection dropdown

diff --git a/Engine/Source/Programs/UnrealDocTool/MarkdownMode/MarkdownMode.csproj b/Engine/Source/Programs/UnrealDocTool/MarkdownMode/MarkdownMode.csproj
--- a/Engine/Source/Programs/UnrealDocTool/MarkdownMode/MarkdownMode.csproj
+++ b/Engine/Source/Programs/UnrealDocTool/MarkdownMode/MarkdownMode.csproj
@@ -196,6 +196,7 @@
     </EmbeddedResource>
   </ItemGroup>
   <ItemGroup>
+    <None Include="app.config" />
     <None Include="Properties\Settings.settings">
       <Generator>SettingsSingleFileGenerator</Generator>
       <LastGenOutput>Settings.Designer.cs</LastGenOutput>
diff --git a/Engine/Source/Programs/UnrealDocTool/MarkdownMode/Properties/Settings.settings b/Engine/Source/Programs/UnrealDocTool/MarkdownMode/Properties/Settings.settings
--- a/Engine/Source/Programs/UnrealDocTool/MarkdownMode/Properties/Settings.settings
+++ b/Engine/Source/Programs/UnrealDocTool/MarkdownMode/Properties/Settings.settings
@@ -56,5 +56,8 @@
     <Setting Name="DisableHighlighting" Type="System.Boolean" Scope="User">
       <Value Profile="(Default)">False</Value>
     </Setting>
+    <Setting Name="SupportedLanguageLabelsString" Type="System.String" Scope="Application">
+      <Value Profile="(Default)">English,中国,日本人,한국어</Value>
+    </Setting>
   </Settings>
 </SettingsFile>
\ No newline at end of file
diff --git a/Engine/Source/Programs/UnrealDocTool/MarkdownMode/Settings.cs b/Engine/Source/Programs/UnrealDocTool/MarkdownMode/Settings.cs
--- a/Engine/Source/Programs/UnrealDocTool/MarkdownMode/Settings.cs
+++ b/Engine/Source/Programs/UnrealDocTool/MarkdownMode/Settings.cs
@@ -13,6 +13,7 @@ namespace MarkdownMode.Properties {
     internal sealed partial class Settings {
         static public System.Drawing.Color DefaultImageFillColor;
         static public string[] SupportedLanguages;
+        static public string[] SupportedLanguageLabels;
         static public string[] MetadataErrorIfMissing;
         static public string[] MetadataInfoIfMissing;
         static public List<string> SupportedAvailabilities;
@@ -26,6 +27,7 @@ namespace MarkdownMode.Properties {
             //
             DefaultImageFillColor = ImageConversion.GetColorFromHexString(DefaultImageFillColorText);
             SupportedLanguages = SupportedLanguagesString.Split(',');
+            SupportedLanguageLabels = SupportedLanguageLabelsString.Split(',');
             MetadataErrorIfMissing = MetadataErrorIfMissingString.Split(',');
             MetadataInfoIfMissing = MetadataInfoIfMissingString.Split(',');
             SupportedAvailabilities = new List<string>(SupportedAvailabilitiesString.Split(','));
diff --git a/Engine/Source/Programs/UnrealDocTool/MarkdownMode/Spellchecker/SpellChecker.Definitions/SpellChecker.Definitions.csproj b/Engine/Source/Programs/UnrealDocTool/MarkdownMode/Spellchecker/SpellChecker.Definitions/SpellChecker.Definitions.csproj
--- a/Engine/Source/Programs/UnrealDocTool/MarkdownMode/Spellchecker/SpellChecker.Definitions/SpellChecker.Definitions.csproj
+++ b/Engine/Source/Programs/UnrealDocTool/MarkdownMode/Spellchecker/SpellChecker.Definitions/SpellChecker.Definitions.csproj
@@ -36,9 +36,18 @@
     <TreatWarningsAsErrors>true</TreatWarningsAsErrors>
   </PropertyGroup>
   <ItemGroup>
-    <Reference Include="Microsoft.VisualStudio.CoreUtility, Version=10.0.0.0, Culture=neutral, PublicKeyToken=b03f5f7f11d50a3a, processorArchitecture=MSIL" />
-    <Reference Include="Microsoft.VisualStudio.Text.Data, Version=10.0.0.0, Culture=neutral, PublicKeyToken=b03f5f7f11d50a3a, processorArchitecture=MSIL" />
-    <Reference Include="Microsoft.VisualStudio.Text.Logic, Version=10.0.0.0, Culture=neutral, PublicKeyToken=b03f5f7f11d50a3a, processorArchitecture=MSIL" />
+    <Reference Include="Microsoft.VisualStudio.CoreUtility, Version=11.0.0.0, Culture=neutral, PublicKeyToken=b03f5f7f11d50a3a, processorArchitecture=MSIL">
+      <SpecificVersion>False</SpecificVersion>
+      <HintPath>C:\Program Files (x86)\Microsoft Visual Studio 11.0\Common7\IDE\CommonExtensions\Microsoft\Editor\Microsoft.VisualStudio.CoreUtility.dll</HintPath>
+    </Reference>
+    <Reference Include="Microsoft.VisualStudio.Text.Data, Version=11.0.0.0, Culture=neutral, PublicKeyToken=b03f5f7f11d50a3a, processorArchitecture=MSIL">
+      <SpecificVersion>False</SpecificVersion>
+      <HintPath>C:\Program Files (x86)\Microsoft Visual Studio 11.0\Common7\IDE\CommonExtensions\Microsoft\Editor\Microsoft.VisualStudio.Text.Data.dll</HintPath>
+    </Reference>
+    <Reference Include="Microsoft.VisualStudio.Text.Logic, Version=11.0.0.0, Culture=neutral, PublicKeyToken=b03f5f7f11d50a3a, processorArchitecture=MSIL">
+      <SpecificVersion>False</SpecificVersion>
+      <HintPath>C:\Program Files (x86)\Microsoft Visual Studio 11.0\Common7\IDE\CommonExtensions\Microsoft\Editor\Microsoft.VisualStudio.Text.Logic.dll</HintPath>
+    </Reference>
     <Reference Include="System" />
     <Reference Include="System.Core" />
     <Reference Include="System.Xml.Linq" />
diff --git a/Engine/Source/Programs/UnrealDocTool/MarkdownMode/ToolWindow/UDNParsingResultsCache.cs b/Engine/Source/Programs/UnrealDocTool/MarkdownMode/ToolWindow/UDNParsingResultsCache.cs
--- a/Engine/Source/Programs/UnrealDocTool/MarkdownMode/ToolWindow/UDNParsingResultsCache.cs
+++ b/Engine/Source/Programs/UnrealDocTool/MarkdownMode/ToolWindow/UDNParsingResultsCache.cs
@@ -163,6 +163,14 @@ namespace MarkdownMode.ToolWindow
                 markdownTransform.ThisIsPreview = true;
 
                 markdownTransform.SupportedLanguages = Settings.SupportedLanguages;
+                markdownTransform.SupportedLanguageLabels = Settings.SupportedLanguageLabels;
+                for (int i = 0; i < markdownTransform.SupportedLanguages.Length; i++)
+                {
+                    if (!markdownTransform.SupportedLanguageMap.ContainsKey(markdownTransform.SupportedLanguages[i]))
+                    {
+                        markdownTransform.SupportedLanguageMap.Add(markdownTransform.SupportedLanguages[i], markdownTransform.SupportedLanguageLabels[i]);
+                    }
+                }
 
                 markdownTransform.AllSupportedAvailability = new List<string>();
                 foreach (string Availablity in Settings.SupportedAvailabilities)
diff --git a/Engine/Source/Programs/UnrealDocTool/MarkdownMode/app.config b/Engine/Source/Programs/UnrealDocTool/MarkdownMode/app.config
--- a/Engine/Source/Programs/UnrealDocTool/MarkdownMode/app.config
+++ b/Engine/Source/Programs/UnrealDocTool/MarkdownMode/app.config
@@ -49,6 +49,9 @@
             <setting name="SupportedAvailabilitiesString" serializeAs="String">
                 <value>Licensee,Rocket,Epic,Docs</value>
             </setting>
+            <setting name="SupportedLanguageLabelsString" serializeAs="String">
+                <value>English,中国,日本人,한국어</value>
+            </setting>
         </MarkdownMode.Properties.Settings>
     </applicationSettings>
     <userSettings>
diff --git a/Engine/Source/Programs/UnrealDocTool/MarkdownMode/source.extension.vsixmanifest b/Engine/Source/Programs/UnrealDocTool/MarkdownMode/source.extension.vsixmanifest
--- a/Engine/Source/Programs/UnrealDocTool/MarkdownMode/source.extension.vsixmanifest
+++ b/Engine/Source/Programs/UnrealDocTool/MarkdownMode/source.extension.vsixmanifest
@@ -1,7 +1,7 @@
 <?xml version="1.0" encoding="utf-8"?>
 <PackageManifest Version="2.0.0" xmlns="http://schemas.microsoft.com/developer/vsx-schema/2011" xmlns:d="http://schemas.microsoft.com/developer/vsx-schema-design/2011">
   <Metadata>
-    <Identity Id="Epic MarkdownMode" Version="0.226" Language="en-US" Publisher="Epic Games" />
+    <Identity Id="Epic MarkdownMode" Version="0.230" Language="en-US" Publisher="Epic Games" />
     <DisplayName>Epic MarkdownMode</DisplayName>
     <Description xml:space="preserve">A language service for editing Markdown files.</Description>
     <Icon>markdown-icon.png</Icon>
diff --git a/Engine/Source/Programs/UnrealDocTool/UnrealDocTool/MarkdownSharp/EpicMarkdown/EMLink.cs b/Engine/Source/Programs/UnrealDocTool/UnrealDocTool/MarkdownSharp/EpicMarkdown/EMLink.cs
--- a/Engine/Source/Programs/UnrealDocTool/UnrealDocTool/MarkdownSharp/EpicMarkdown/EMLink.cs
+++ b/Engine/Source/Programs/UnrealDocTool/UnrealDocTool/MarkdownSharp/EpicMarkdown/EMLink.cs
@@ -6,11 +6,13 @@ using System.Text;
 using System.Drawing;
 using System.IO;
 using System.Text.RegularExpressions;
+using System.Globalization;
 
 using DotLiquid;
 
 using MarkdownSharp.EpicMarkdown.Parsers;
 using MarkdownSharp.EpicMarkdown.PathProviders;
+using MarkdownSharp.Preprocessor;
 
 namespace MarkdownSharp.EpicMarkdown
 {
@@ -21,14 +23,18 @@ namespace MarkdownSharp.EpicMarkdown
 
         public string Title { get; private set; }
         public EMSpanElements Content { get; private set; }
+        public bool isFancyLink { get; private set; }
+        public Hash metadata { get; private set; }
 
-        public EMLink(EMDocument doc, EMElementOrigin origin, EMElement parent, TransformationData data, EMPathProvider path, string title, EMSpanElements content)
+        public EMLink(EMDocument doc, EMElementOrigin origin, EMElement parent, TransformationData data, EMPathProvider path, string title, EMSpanElements content, bool isFancyLinkIn, Hash metadataIn)
             : base(doc, origin, parent)
         {
             Path = path;
 
             Title = title;
             Content = content;
+            isFancyLink = isFancyLinkIn;
+            metadata = metadataIn;
 
             if (path.ChangedLanguage)
             {
@@ -117,6 +123,11 @@ namespace MarkdownSharp.EpicMarkdown
             try
             {
                 bool isAPILink = false;
+                bool isFancyLinkLocal = false;
+                Hash metadataLocal = Hash.FromAnonymousObject(
+                    new
+                    { 
+                    });
                 string APIMember = "";
                 if (path.StartsWith("API:"))
                 {
@@ -132,6 +143,12 @@ namespace MarkdownSharp.EpicMarkdown
                     }
                     isAPILink = true;
                 }
+                else if (path.StartsWith("FANCY:"))
+                {
+                    path = path.Substring(6, path.Length - 6);
+                    isFancyLinkLocal = true;
+                    metadataLocal = ProcessFancyLink(path, data);
+                }
                 var pathProvider = EMPathProvider.CreatePath(path, doc, data);
                 var errorId = -1;
 
@@ -218,11 +235,11 @@ namespace MarkdownSharp.EpicMarkdown
 
                 if (isAPILink)
                 {
-                    return new EMLink(doc, orig, parent, data, pathProvider, APIMember, content);
+                    return new EMLink(doc, orig, parent, data, pathProvider, APIMember, content, false, metadataLocal);
                 }
                 else
                 {
-                    return new EMLink(doc, orig, parent, data, pathProvider, title, content);
+                    return new EMLink(doc, orig, parent, data, pathProvider, title, content, isFancyLinkLocal, metadataLocal);
                 }
             }
             catch (EMPathVerificationException e)
@@ -259,6 +276,32 @@ namespace MarkdownSharp.EpicMarkdown
             return Create(doc, orig, parent, data, content, url, GetReferenceTitle(data, linkId));
         }
 
+        public static Hash ProcessFancyLink(string path, TransformationData data)
+        {
+            ClosestFileStatus status;
+            var linkedDoc = data.ProcessedDocumentCache.GetClosest(path, out status);
+
+            var parameters = Hash.FromAnonymousObject(
+                new
+                {
+                });
+
+            foreach (var metadata in linkedDoc.PreprocessedData.Metadata.MetadataMap)
+            {
+                if(metadata.Key == "skilllevel")
+                {
+                    parameters.Add("skilllevellabel", string.Join(", ", metadata.Value));
+                    for(int i=0; i < metadata.Value.Count; i++)
+                    {
+                        metadata.Value[i] = metadata.Value[i].ToLower();
+                    }
+                }
+                parameters.Add(metadata.Key, string.Join(", ", metadata.Value));
+            }
+
+            return parameters;
+        }
+
         public override void ForEachWithContext<T>(T context)
         {
             context.PreChildrenVisit(this);
@@ -284,12 +327,17 @@ namespace MarkdownSharp.EpicMarkdown
                     builder.Append(linkText);
                     return;
                 }
-                
-                var result =
-                    Templates.Link.Render(
+
+                var result = "";
+                if (isFancyLink)
+                {
+                    result = Templates.FancyLink.Render(
                         Hash.FromAnonymousObject(
                             new
                                 {
+                                    version = metadata["version"],
+                                    skilllevel = metadata["skilllevel"],
+                                    skilllevellabel = metadata["skilllevellabel"],
                                     linkUrl = Path.GetPath(data).Replace("\\", "/"),
                                     linkTitle =
                                         !String.IsNullOrWhiteSpace(Title) ? Markdown.EscapeBoldItalic(Title) : null,
@@ -308,8 +356,37 @@ namespace MarkdownSharp.EpicMarkdown
                                                         "Images",
                                                         "INT_flag.jpg")
                                                           }))
-                                            : linkText
+                                            : linkText,
+                                    relativeHTMLPath = data.CurrentFolderDetails.RelativeHTMLPath
                                 }));
+                }
+                else
+                {
+                    result = Templates.Link.Render(
+                        Hash.FromAnonymousObject(
+                            new
+                            {
+                                linkUrl = Path.GetPath(data).Replace("\\", "/"),
+                                linkTitle =
+                                    !String.IsNullOrWhiteSpace(Title) ? Markdown.EscapeBoldItalic(Title) : null,
+                                linkText =
+                                    Path.ChangedLanguage
+                                        ? linkText
+                                          + Templates.ImageFrame.Render(
+                                              Hash.FromAnonymousObject(
+                                                  new
+                                                  {
+                                                      imageClass = "languageinline",
+                                                      imagePath =
+                                            System.IO.Path.Combine(
+                                                data.CurrentFolderDetails.RelativeHTMLPath,
+                                                "Include",
+                                                "Images",
+                                                "INT_flag.jpg")
+                                                  }))
+                                        : linkText
+                            }));
+                }
 
                 if (Path.ChangedLanguage)
                 {
diff --git a/Engine/Source/Programs/UnrealDocTool/UnrealDocTool/MarkdownSharp/Markdown.cs b/Engine/Source/Programs/UnrealDocTool/UnrealDocTool/MarkdownSharp/Markdown.cs
--- a/Engine/Source/Programs/UnrealDocTool/UnrealDocTool/MarkdownSharp/Markdown.cs
+++ b/Engine/Source/Programs/UnrealDocTool/UnrealDocTool/MarkdownSharp/Markdown.cs
@@ -377,6 +377,7 @@ namespace MarkdownSharp
         /// </summary>
         public Markdown(bool loadOptionsFromConfigFile)
         {
+            SupportedLanguageMap = new Dictionary<string, string>();
             JustReadVariableDefinitionsRun = false;
 
             Preprocessor = new Preprocessor.Preprocessor(this);
@@ -415,6 +416,7 @@ namespace MarkdownSharp
         /// </summary>
         public Markdown(MarkdownOptions options)
         {
+            SupportedLanguageMap = new Dictionary<string, string>();
             _autoHyperlink = options.AutoHyperlink;
             _autoNewlines = options.AutoNewlines;
             _emptyElementSuffix = options.EmptyElementSuffix;
@@ -543,9 +545,15 @@ namespace MarkdownSharp
         // Internal counter incremented when stepping to a new list level and decremented on the way out.
         private int ListLevel { get; set; }
 
+        // What languages are supported
+        public Dictionary<string, string> SupportedLanguageMap { get; set; }
+
         // What languages are supported
         public string[] SupportedLanguages { get; set; }
 
+        // What languages are supported
+        public string[] SupportedLanguageLabels { get; set; }
+
         // Error if this metadata value is missing from a file
         public static string[] MetadataErrorIfMissing { get; set; }
 
@@ -1088,6 +1096,7 @@ namespace MarkdownSharp
             }
 
             var translatedPageLinks = "";
+            var translatedPageSelect = "";
 
             //Generate links to other languages for this page
             foreach (var translatedLanguage in data.LanguagesLinksToGenerate)
@@ -1102,6 +1111,7 @@ namespace MarkdownSharp
                                         ? ""
                                         : data.CurrentFolderDetails.CurrentFolderFromMarkdownAsTopLeaf.Replace("\\", "/"),
                                 otherLanguage = translatedLanguage,
+                                otherLanguageLabel = SupportedLanguageMap[translatedLanguage],
                                 relativeHTMLPath = data.CurrentFolderDetails.RelativeHTMLPath
                             });
 
@@ -1119,8 +1129,15 @@ namespace MarkdownSharp
                 }
 
                 // Cope with top level folders having blank current folder.
-                translatedPageLinks += 
+                translatedPageLinks +=
                     Templates.TranslatedPageLink.Render(linkParams);
+                var linkPageParams =
+                    Hash.FromAnonymousObject(
+                        new
+                        {
+                            languageLinks = translatedPageLinks
+                        });
+                translatedPageSelect = Templates.TranslatedPages.Render(linkPageParams);
             }
 
             //add a warning if crumbs for the document have not been updated to new format
@@ -1148,7 +1165,7 @@ namespace MarkdownSharp
                         title = data.ProcessedDocumentCache.CurrentFileDocument.Metadata.DocumentTitle,
                         seotitle = (data.ProcessedDocumentCache.CurrentFileDocument.Metadata.SEOTitle != null) ? data.ProcessedDocumentCache.CurrentFileDocument.Metadata.SEOTitle : data.ProcessedDocumentCache.CurrentFileDocument.Metadata.DocumentTitle,
                         metadata = metadata,
-                        translatedPages = translatedPageLinks,
+                        translatedPages = translatedPageSelect,
                         relatedPagesMenu = Templates.RelatedPages.Render(Hash.FromAnonymousObject(
                             new {
                                 relatedPages = data.ProcessedDocumentCache.Metadata.RelatedLinks,
diff --git a/Engine/Source/Programs/UnrealDocTool/UnrealDocTool/MarkdownSharp/Preprocessor/MetadataManager.cs b/Engine/Source/Programs/UnrealDocTool/UnrealDocTool/MarkdownSharp/Preprocessor/MetadataManager.cs
--- a/Engine/Source/Programs/UnrealDocTool/UnrealDocTool/MarkdownSharp/Preprocessor/MetadataManager.cs
+++ b/Engine/Source/Programs/UnrealDocTool/UnrealDocTool/MarkdownSharp/Preprocessor/MetadataManager.cs
@@ -107,7 +107,7 @@ namespace MarkdownSharp.Preprocessor
 
                 if (metaDataCategoryLowerCase == "skilllevel")
                 {
-                    SkillLevels.Add(metaValue);
+                    SkillLevels.Add(CultureInfo.CurrentCulture.TextInfo.ToTitleCase(metaValue));
                 }
 
                 if (metaDataCategoryLowerCase == "tags")
diff --git a/Engine/Source/Programs/UnrealDocTool/UnrealDocTool/MarkdownSharp/Templates.cs b/Engine/Source/Programs/UnrealDocTool/UnrealDocTool/MarkdownSharp/Templates.cs
--- a/Engine/Source/Programs/UnrealDocTool/UnrealDocTool/MarkdownSharp/Templates.cs
+++ b/Engine/Source/Programs/UnrealDocTool/UnrealDocTool/MarkdownSharp/Templates.cs
@@ -235,6 +235,11 @@ namespace MarkdownSharp
             get { return GetCached("link.html"); }
         }
 
+        public static TemplateFile FancyLink
+        {
+            get { return GetCached("fancyLink.html"); }
+        }
+
         public static TemplateFile Header
         {
             get { return GetCached("header.html"); }
@@ -340,6 +345,11 @@ namespace MarkdownSharp
             get { return GetCached("translatedPageLink.html"); }
         }
 
+        public static TemplateFile TranslatedPages
+        {
+            get { return GetCached("translatedPages.html"); }
+        }
+
         public static TemplateFile RelatedPages
         {
             get { return GetCached("relatedPages.html"); }
diff --git a/Engine/Source/Programs/UnrealDocTool/UnrealDocTool/MarkdownSharpTests/MarkdownTests/MarkdownTest.cs b/Engine/Source/Programs/UnrealDocTool/UnrealDocTool/MarkdownSharpTests/MarkdownTests/MarkdownTest.cs
--- a/Engine/Source/Programs/UnrealDocTool/UnrealDocTool/MarkdownSharpTests/MarkdownTests/MarkdownTest.cs
+++ b/Engine/Source/Programs/UnrealDocTool/UnrealDocTool/MarkdownSharpTests/MarkdownTests/MarkdownTest.cs
@@ -331,6 +331,11 @@ namespace MarkdownSharpTests.MarkdownTests
             markdown.ThisIsPreview = false;
 
             markdown.SupportedLanguages = Settings.Default.SupportedLanguages.ToUpper().Split(',');
+            markdown.SupportedLanguageLabels = Settings.Default.SupportedLanguageLabels.Split(',');
+            for(int i = 0; i < markdown.SupportedLanguages.Length ; i++)
+            {
+                markdown.SupportedLanguageMap.Add(markdown.SupportedLanguages[i], markdown.SupportedLanguageLabels[i]);
+            }
 
             //Pass the default conversion settings to Markdown for use in the imagedetails creation.
             markdown.DefaultImageDoCompress = Settings.Default.DoCompressImages;
diff --git a/Engine/Source/Programs/UnrealDocTool/UnrealDocTool/MarkdownSharpTests/Settings.settings b/Engine/Source/Programs/UnrealDocTool/UnrealDocTool/MarkdownSharpTests/Settings.settings
--- a/Engine/Source/Programs/UnrealDocTool/UnrealDocTool/MarkdownSharpTests/Settings.settings
+++ b/Engine/Source/Programs/UnrealDocTool/UnrealDocTool/MarkdownSharpTests/Settings.settings
@@ -50,5 +50,8 @@
     <Setting Name="InternationalizationDirectory" Type="System.String" Scope="Application">
       <Value Profile="(Default)">Internationalization</Value>
     </Setting>
+    <Setting Name="SupportedLanguageLabels" Type="System.String" Scope="Application">
+      <Value Profile="(Default)">English,中国,日本人,한국어</Value>
+    </Setting>
   </Settings>
 </SettingsFile>
\ No newline at end of file
diff --git a/Engine/Source/Programs/UnrealDocTool/UnrealDocTool/MarkdownSharpTests/app.config b/Engine/Source/Programs/UnrealDocTool/UnrealDocTool/MarkdownSharpTests/app.config
--- a/Engine/Source/Programs/UnrealDocTool/UnrealDocTool/MarkdownSharpTests/app.config
+++ b/Engine/Source/Programs/UnrealDocTool/UnrealDocTool/MarkdownSharpTests/app.config
@@ -6,6 +6,7 @@
       <section name="UnrealDocTool.Properties.Settings" type="System.Configuration.ClientSettingsSection, System, Version=4.0.0.0, Culture=neutral, PublicKeyToken=b77a5c561934e089" allowExeDefinition="MachineToLocalUser" requirePermission="false"/>
     </sectionGroup>
     <sectionGroup name="applicationSettings" type="System.Configuration.ApplicationSettingsGroup, System, Version=4.0.0.0, Culture=neutral, PublicKeyToken=b77a5c561934e089">
+      <section name="MarkdownSharpTests.Settings" type="System.Configuration.ClientSettingsSection, System, Version=4.0.0.0, Culture=neutral, PublicKeyToken=b77a5c561934e089" requirePermission="false" />
       <section name="UnrealDocTool.Properties.Settings" type="System.Configuration.ClientSettingsSection, System, Version=4.0.0.0, Culture=neutral, PublicKeyToken=b77a5c561934e089" requirePermission="false"/>
     </sectionGroup>
     <sectionGroup name="userSettings" type="System.Configuration.UserSettingsGroup, System, Version=4.0.0.0, Culture=neutral, PublicKeyToken=b77a5c561934e089">
@@ -26,6 +27,59 @@
   </system.diagnostics>
   
   <applicationSettings>
+    <MarkdownSharpTests.Settings>
+      <setting name="SupportedAvailability" serializeAs="String">
+        <value>Licensee,Rocket,Epic</value>
+      </setting>
+      <setting name="SourceDirectory" serializeAs="String">
+        <value>../../../../../../../Documentation/Source</value>
+      </setting>
+      <setting name="OutputDirectory" serializeAs="String">
+        <value>../../../../../../../Documentation/HTML</value>
+      </setting>
+      <setting name="DoCompressImages" serializeAs="String">
+        <value>True</value>
+      </setting>
+      <setting name="CompressImageType" serializeAs="String">
+        <value>jpg</value>
+      </setting>
+      <setting name="ImageJPGCompressionValue" serializeAs="String">
+        <value>80</value>
+      </setting>
+      <setting name="ImageFillColor" serializeAs="String">
+        <value>#FFFFFF</value>
+      </setting>
+      <setting name="SupportedLanguages" serializeAs="String">
+        <value>INT,CHN,JPN,KOR</value>
+      </setting>
+      <setting name="DefaultTemplate" serializeAs="String">
+        <value>Default.html</value>
+      </setting>
+      <setting name="MetadataErrorIfMissing" serializeAs="String">
+        <value>Title,Crumbs,Availability</value>
+      </setting>
+      <setting name="MetadataInfoIfMissing" serializeAs="String">
+        <value>Description</value>
+      </setting>
+      <setting name="MetadataCleanKeepThese" serializeAs="String">
+        <value>Title,Crumbs,Description,HtmlDocumentType,INTSourceChangelist,Availability,DoIndex,Keywords</value>
+      </setting>
+      <setting name="PublicAvailabilityFlag" serializeAs="String">
+        <value>Public</value>
+      </setting>
+      <setting name="APIFolderLocation" serializeAs="String">
+        <value>API</value>
+      </setting>
+      <setting name="IncludeDirectory" serializeAs="String">
+        <value>Include</value>
+      </setting>
+      <setting name="InternationalizationDirectory" serializeAs="String">
+        <value>Internationalization</value>
+      </setting>
+      <setting name="SupportedLanguageLabels" serializeAs="String">
+        <value>English,中国,日本人,한국어</value>
+      </setting>
+    </MarkdownSharpTests.Settings>
     <UnrealDocTool.Properties.Settings>
       <setting name="SourceDirectory" serializeAs="String">
         <value>../../Documentation/Source</value>
diff --git a/Engine/Source/Programs/UnrealDocTool/UnrealDocTool/UnrealDocTool/Program.cs b/Engine/Source/Programs/UnrealDocTool/UnrealDocTool/UnrealDocTool/Program.cs
--- a/Engine/Source/Programs/UnrealDocTool/UnrealDocTool/UnrealDocTool/Program.cs
+++ b/Engine/Source/Programs/UnrealDocTool/UnrealDocTool/UnrealDocTool/Program.cs
@@ -51,6 +51,8 @@ namespace UnrealDocTool
         private static string DefaultTemplate;
 
         private static string[] SupportedLanguages;
+        private static string[] SupportedLanguageLabels;
+        private static Dictionary<string, string> SupportedLanguageMap = new Dictionary<string,string>();
         private static string[] MetadataErrorIfMissing;
         private static string[] MetadataInfoIfMissing;
         private static string[] MetadataCleanKeepThese;
@@ -223,6 +225,11 @@ namespace UnrealDocTool
             }
 
             SupportedLanguages = Settings.Default.SupportedLanguages.ToUpper().Split(',');
+            SupportedLanguageLabels = Settings.Default.SupportedLanguageLabels.Split(',');
+            for (int i = 0; i < SupportedLanguages.Length; i++)
+            {
+                SupportedLanguageMap.Add(SupportedLanguages[i], SupportedLanguageLabels[i]);
+            }
 
             MetadataErrorIfMissing = Settings.Default.MetadataErrorIfMissing.Split(',');
             MetadataInfoIfMissing = Settings.Default.MetadataInfoIfMissing.Split(',');
@@ -660,6 +667,14 @@ namespace UnrealDocTool
             }
 
             markdownToHtml.SupportedLanguages = SupportedLanguages;
+            markdownToHtml.SupportedLanguageLabels = SupportedLanguageLabels;
+            for (int i = 0; i < markdownToHtml.SupportedLanguages.Length; i++)
+            {
+                if(!markdownToHtml.SupportedLanguageMap.ContainsKey(SupportedLanguages[i]))
+                {
+                    markdownToHtml.SupportedLanguageMap.Add(markdownToHtml.SupportedLanguages[i], markdownToHtml.SupportedLanguageLabels[i]);
+                }
+            }
 
             //Pass the default conversion settings to Markdown for use in the image details creation.
             markdownToHtml.DefaultImageDoCompress = DoCompressImages;
diff --git a/Engine/Source/Programs/UnrealDocTool/UnrealDocTool/UnrealDocTool/Properties/Settings.settings b/Engine/Source/Programs/UnrealDocTool/UnrealDocTool/UnrealDocTool/Properties/Settings.settings
--- a/Engine/Source/Programs/UnrealDocTool/UnrealDocTool/UnrealDocTool/Properties/Settings.settings
+++ b/Engine/Source/Programs/UnrealDocTool/UnrealDocTool/UnrealDocTool/Properties/Settings.settings
@@ -26,9 +26,6 @@
     <Setting Name="DefaultTemplate" Type="System.String" Scope="Application">
       <Value Profile="(Default)">Default.html</Value>
     </Setting>
-    <Setting Name="MetadataErrorIfMissing" Type="System.String" Scope="Application">
-      <Value Profile="(Default)">Title,Crumbs,Availability</Value>
-    </Setting>
     <Setting Name="MetadataInfoIfMissing" Type="System.String" Scope="Application">
       <Value Profile="(Default)">Description</Value>
     </Setting>
@@ -53,5 +50,11 @@
     <Setting Name="DoNotPublishAvailabilityFlag" Type="System.String" Scope="Application">
       <Value Profile="(Default)">NoPublish</Value>
     </Setting>
+    <Setting Name="SupportedLanguageLabels" Type="System.String" Scope="Application">
+      <Value Profile="(Default)">English,中国,日本人,한국어</Value>
+    </Setting>
+    <Setting Name="MetadataErrorIfMissing" Type="System.String" Scope="Application">
+      <Value Profile="(Default)">Title,Availability</Value>
+    </Setting>
   </Settings>
 </SettingsFile>
\ No newline at end of file
diff --git a/Engine/Source/Programs/UnrealDocTool/UnrealDocTool/UnrealDocTool/app.config b/Engine/Source/Programs/UnrealDocTool/UnrealDocTool/UnrealDocTool/app.config
--- a/Engine/Source/Programs/UnrealDocTool/UnrealDocTool/UnrealDocTool/app.config
+++ b/Engine/Source/Programs/UnrealDocTool/UnrealDocTool/UnrealDocTool/app.config
@@ -51,9 +51,6 @@
       <setting name="DefaultTemplate" serializeAs="String">
         <value>Default.html</value>
       </setting>
-      <setting name="MetadataErrorIfMissing" serializeAs="String">
-        <value>Title,Availability</value>
-      </setting>
       <setting name="MetadataInfoIfMissing" serializeAs="String">
         <value>Description</value>
       </setting>
@@ -78,6 +75,12 @@
       <setting name="DoNotPublishAvailabilityFlag" serializeAs="String">
         <value>NoPublish</value>
       </setting>
+      <setting name="SupportedLanguageLabels" serializeAs="String">
+        <value>English,中国,日本人,한국어</value>
+      </setting>
+      <setting name="MetadataErrorIfMissing" serializeAs="String">
+        <value>Title,Availability</value>
+      </setting>
     </UnrealDocTool.Properties.Settings>
   </applicationSettings>
   <userSettings>
@@ -94,6 +97,9 @@
       <setting name="SupportedLanguages" serializeAs="String">
         <value>INT,CHN,JPN,KOR</value>
       </setting>
+      <setting name="SupportedLanguageLabels" serializeAs="String">
+        <value>English,中国,日本人,한국어</value>
+      </setting>
       <setting name="HasRun" serializeAs="String">
         <value>False</value>
       </setting>
diff --git a/Engine/Source/Programs/UnrealDocTool/UnrealDocTool/UnrealDocToolDoxygenInputFilter/UnrealDocToolDoxygenInputFilter.csproj b/Engine/Source/Programs/UnrealDocTool/UnrealDocTool/UnrealDocToolDoxygenInputFilter/UnrealDocToolDoxygenInputFilter.csproj
--- a/Engine/Source/Programs/UnrealDocTool/UnrealDocTool/UnrealDocToolDoxygenInputFilter/UnrealDocToolDoxygenInputFilter.csproj
+++ b/Engine/Source/Programs/UnrealDocTool/UnrealDocTool/UnrealDocToolDoxygenInputFilter/UnrealDocToolDoxygenInputFilter.csproj
@@ -51,7 +51,9 @@
     <Compile Include="Properties\AssemblyInfo.cs" />
   </ItemGroup>
   <ItemGroup>
-    <None Include="App.config" />
+    <None Include="App.config">
+      <SubType>Designer</SubType>
+    </None>
   </ItemGroup>
   <Import Project="$(MSBuildToolsPath)\Microsoft.CSharp.targets" />
   <!-- To modify your build process, add your task inside one of the targets below and uncomment it. 
