Disable XGE when building DocumentationEditor. It's only building UnrealHeaderTool, so it's not much use anyway - but it has caused failures on the build farm.

diff --git a/Engine/Source/Programs/UnrealDocTool/APIDocTool/APIDocTool/Program.cs b/Engine/Source/Programs/UnrealDocTool/APIDocTool/APIDocTool/Program.cs
--- a/Engine/Source/Programs/UnrealDocTool/APIDocTool/APIDocTool/Program.cs
+++ b/Engine/Source/Programs/UnrealDocTool/APIDocTool/APIDocTool/Program.cs
@@ -632,7 +632,7 @@ namespace APIDocTool
 				Console.WriteLine("Building target info...");
 				Utility.SafeCreateDirectory(Path.GetDirectoryName(TargetInfoPath));
 
-				string Arguments = String.Format("DocumentationEditor Win64 Debug -project=\"{0}\"", Path.Combine(EngineDir, "Documentation\\Extras\\API\\Build\\Documentation.uproject"));
+				string Arguments = String.Format("DocumentationEditor Win64 Debug -noxge -project=\"{0}\"", Path.Combine(EngineDir, "Documentation\\Extras\\API\\Build\\Documentation.uproject"));
 				if (!RunUnrealBuildTool(EngineDir, Arguments + " -clean"))
 				{
 					return false;
