Improved Python script (for archival purposes only).

diff --git a/Engine/Source/Programs/UnrealDocTool/DoxygenInputFilter/CommentFixup.py b/Engine/Source/Programs/UnrealDocTool/DoxygenInputFilter/CommentFixup.py
--- a/Engine/Source/Programs/UnrealDocTool/DoxygenInputFilter/CommentFixup.py
+++ b/Engine/Source/Programs/UnrealDocTool/DoxygenInputFilter/CommentFixup.py
@@ -1,11 +1,28 @@
+#!/usr/bin/env python
 import sys
+
+# Scan LINE for TYPE1 and TYPE2, report the position of whichever comes first. If neither is found, return length of LINE.
+def detectfirstcomment(line, type1, type2)
+    len = line.len()
+    pos = line.find(type1)
+    pos2 = line.find(type2)
+    if (pos >= 0)
+        pos = len
+    if (pos2 >= 0)
+        pos2 = len
+    pos = min(pos, pos2)
+    return pos
+
 f = file(sys.argv[1], 'r')
 line = f.readline()
 while line:
-    # Pad the string with an extra space at the end to catch end-of-line block comments.
-    line = line + " "
-    line = line.replace("/* ", "/** ")
-    line = line.replace("// ", "/// ")
+    # If a harvestable tag shows up earlier in the line than an explicit ignore tag, the harvestable tag takes precedence.
+    posignore = detectfirstcomment(line, "//~", "/*~")
+    posharvest = detectfirstcomment(line "//", "/*")
+    if (posharvest < posignore)
+        # When a comment is determined to be harvestable, ensure that it's formatted properly for Doxygen.
+        line = line.replace("/*", "/**")
+        line = line.replace("//", "///")
     sys.stdout.write(line)
     line = f.readline()
-f.close() 
\ No newline at end of file
+f.close()
\ No newline at end of file
