Changing APIDocTool's filtering system to use a C# program instead of a Python script due to build machines not supporting Python.

diff --git a/Engine/Source/Programs/UnrealDocTool/DoxygenInputFilter/App.config b/Engine/Source/Programs/UnrealDocTool/DoxygenInputFilter/App.config
new file mode 100644
--- /dev/null
+++ b/Engine/Source/Programs/UnrealDocTool/DoxygenInputFilter/App.config
@@ -0,0 +1,6 @@
+﻿<?xml version="1.0" encoding="utf-8" ?>
+<configuration>
+    <startup> 
+        <supportedRuntime version="v4.0" sku=".NETFramework,Version=v4.5" />
+    </startup>
+</configuration>
\ No newline at end of file
diff --git a/Engine/Source/Programs/UnrealDocTool/DoxygenInputFilter/CommentFixup.py b/Engine/Source/Programs/UnrealDocTool/DoxygenInputFilter/CommentFixup.py
deleted file mode 100644
diff --git a/Engine/Source/Programs/UnrealDocTool/DoxygenInputFilter/DoxygenInputFilter.csproj b/Engine/Source/Programs/UnrealDocTool/DoxygenInputFilter/DoxygenInputFilter.csproj
new file mode 100644
--- /dev/null
+++ b/Engine/Source/Programs/UnrealDocTool/DoxygenInputFilter/DoxygenInputFilter.csproj
@@ -0,0 +1,58 @@
+﻿<?xml version="1.0" encoding="utf-8"?>
+<Project ToolsVersion="12.0" DefaultTargets="Build" xmlns="http://schemas.microsoft.com/developer/msbuild/2003">
+  <Import Project="$(MSBuildExtensionsPath)\$(MSBuildToolsVersion)\Microsoft.Common.props" Condition="Exists('$(MSBuildExtensionsPath)\$(MSBuildToolsVersion)\Microsoft.Common.props')" />
+  <PropertyGroup>
+    <Configuration Condition=" '$(Configuration)' == '' ">Debug</Configuration>
+    <Platform Condition=" '$(Platform)' == '' ">AnyCPU</Platform>
+    <ProjectGuid>{198E4BA5-29E1-44DE-A471-5316DFFE0584}</ProjectGuid>
+    <OutputType>Exe</OutputType>
+    <AppDesignerFolder>Properties</AppDesignerFolder>
+    <RootNamespace>DoxygenInputFilter</RootNamespace>
+    <AssemblyName>DoxygenInputFilter</AssemblyName>
+    <TargetFrameworkVersion>v4.5</TargetFrameworkVersion>
+    <FileAlignment>512</FileAlignment>
+  </PropertyGroup>
+  <PropertyGroup Condition=" '$(Configuration)|$(Platform)' == 'Debug|AnyCPU' ">
+    <PlatformTarget>AnyCPU</PlatformTarget>
+    <DebugSymbols>true</DebugSymbols>
+    <DebugType>full</DebugType>
+    <Optimize>false</Optimize>
+    <OutputPath>bin\Debug\</OutputPath>
+    <DefineConstants>DEBUG;TRACE</DefineConstants>
+    <ErrorReport>prompt</ErrorReport>
+    <WarningLevel>4</WarningLevel>
+  </PropertyGroup>
+  <PropertyGroup Condition=" '$(Configuration)|$(Platform)' == 'Release|AnyCPU' ">
+    <PlatformTarget>AnyCPU</PlatformTarget>
+    <DebugType>pdbonly</DebugType>
+    <Optimize>true</Optimize>
+    <OutputPath>bin\Release\</OutputPath>
+    <DefineConstants>TRACE</DefineConstants>
+    <ErrorReport>prompt</ErrorReport>
+    <WarningLevel>4</WarningLevel>
+  </PropertyGroup>
+  <ItemGroup>
+    <Reference Include="System" />
+    <Reference Include="System.Core" />
+    <Reference Include="System.Xml.Linq" />
+    <Reference Include="System.Data.DataSetExtensions" />
+    <Reference Include="Microsoft.CSharp" />
+    <Reference Include="System.Data" />
+    <Reference Include="System.Xml" />
+  </ItemGroup>
+  <ItemGroup>
+    <Compile Include="Program.cs" />
+    <Compile Include="Properties\AssemblyInfo.cs" />
+  </ItemGroup>
+  <ItemGroup>
+    <None Include="App.config" />
+  </ItemGroup>
+  <Import Project="$(MSBuildToolsPath)\Microsoft.CSharp.targets" />
+  <!-- To modify your build process, add your task inside one of the targets below and uncomment it. 
+       Other similar extension points exist, see Microsoft.Common.targets.
+  <Target Name="BeforeBuild">
+  </Target>
+  <Target Name="AfterBuild">
+  </Target>
+  -->
+</Project>
\ No newline at end of file
diff --git a/Engine/Source/Programs/UnrealDocTool/DoxygenInputFilter/Program.cs b/Engine/Source/Programs/UnrealDocTool/DoxygenInputFilter/Program.cs
new file mode 100644
--- /dev/null
+++ b/Engine/Source/Programs/UnrealDocTool/DoxygenInputFilter/Program.cs
@@ -0,0 +1,63 @@
+﻿using System;
+using System.Collections.Generic;
+using System.Linq;
+using System.Text;
+using System.Threading.Tasks;
+using System.IO;
+
+namespace DoxygenInputFilter
+{
+	class Program
+	{
+		static void Main(string[] args)
+		{
+			StreamReader reader = File.OpenText(args[0]);
+			string line;
+			while ((line = reader.ReadLine()) != null)
+			{
+				int posIgnore = DetectFirstSubstring(line, "//~", "/*~");
+				int posHarvest = DetectFirstSubstring(line, "//", "/*");
+				if (posHarvest < posIgnore)
+				{
+					// This line does not contain an ignore code, or has a comment section starter before it. It should be harvested.
+					int posDoxygenHarvest = DetectFirstSubstring(line, "///", "/**");
+					if (posHarvest < posDoxygenHarvest)
+					{
+						// This line would normally be ignored by Doxygen, so modify it into something Doxygen will harvest.
+						Console.Out.WriteLine(ReplaceFirstSubstring(ReplaceFirstSubstring(line, "/*", "/**"), "//", "///"));
+						continue;
+					}
+				}
+				Console.Out.WriteLine(line);
+			}
+		}
+
+		static int DetectFirstSubstring(string sourceLine, string substring1, string substring2)
+		{
+			int len = sourceLine.Length;
+			int pos1 = sourceLine.IndexOf(substring1);
+			int pos2 = sourceLine.IndexOf(substring2);
+			if (pos1 < 0)
+			{
+				pos1 = len;
+			}
+			if (pos2 < 0)
+			{
+				pos2 = len;
+			}
+			return Math.Min(pos1, pos2);
+		}
+
+		static string ReplaceFirstSubstring(string sourceString, string searchSubstring, string replacementSubstring)
+		{
+			int pos = sourceString.IndexOf(searchSubstring);
+			if (pos < 0)
+			{
+				// Not found
+				return sourceString;
+			}
+			// Manually cut the first instance of searchString out and replace it.
+			return (sourceString.Substring(0, pos) + replacementSubstring + sourceString.Substring(pos + searchSubstring.Length));
+		}
+	}
+}
diff --git a/Engine/Source/Programs/UnrealDocTool/DoxygenInputFilter/Properties/AssemblyInfo.cs b/Engine/Source/Programs/UnrealDocTool/DoxygenInputFilter/Properties/AssemblyInfo.cs
new file mode 100644
--- /dev/null
+++ b/Engine/Source/Programs/UnrealDocTool/DoxygenInputFilter/Properties/AssemblyInfo.cs
@@ -0,0 +1,36 @@
+﻿using System.Reflection;
+using System.Runtime.CompilerServices;
+using System.Runtime.InteropServices;
+
+// General Information about an assembly is controlled through the following 
+// set of attributes. Change these attribute values to modify the information
+// associated with an assembly.
+[assembly: AssemblyTitle("DoxygenInputFilter")]
+[assembly: AssemblyDescription("")]
+[assembly: AssemblyConfiguration("")]
+[assembly: AssemblyCompany("")]
+[assembly: AssemblyProduct("DoxygenInputFilter")]
+[assembly: AssemblyCopyright("Copyright ©  2015")]
+[assembly: AssemblyTrademark("")]
+[assembly: AssemblyCulture("")]
+
+// Setting ComVisible to false makes the types in this assembly not visible 
+// to COM components.  If you need to access a type in this assembly from 
+// COM, set the ComVisible attribute to true on that type.
+[assembly: ComVisible(false)]
+
+// The following GUID is for the ID of the typelib if this project is exposed to COM
+[assembly: Guid("b361b765-d54e-44e4-8062-d36a56fe9ccd")]
+
+// Version information for an assembly consists of the following four values:
+//
+//      Major Version
+//      Minor Version 
+//      Build Number
+//      Revision
+//
+// You can specify all the values or you can default the Build and Revision Numbers 
+// by using the '*' as shown below:
+// [assembly: AssemblyVersion("1.0.*")]
+[assembly: AssemblyVersion("1.0.0.0")]
+[assembly: AssemblyFileVersion("1.0.0.0")]
diff --git a/Engine/Source/Programs/UnrealDocTool/DoxygenLib/Doxygen.cs b/Engine/Source/Programs/UnrealDocTool/DoxygenLib/Doxygen.cs
--- a/Engine/Source/Programs/UnrealDocTool/DoxygenLib/Doxygen.cs
+++ b/Engine/Source/Programs/UnrealDocTool/DoxygenLib/Doxygen.cs
@@ -70,16 +70,14 @@ namespace DoxygenLib
 
 				// Input filter program
 				Output.WriteLine();
-				//FormatSetting(Output, "INPUT_FILTER", Path.GetDirectoryName(ExecutePath) + "/../../../../../DoxygenInputFilter/CommentFixup.py");
 				string FilterPath = Path.GetDirectoryName(ExecutePath).Replace("\\", "/");
-				// Back up by five levels, then tack on the relative path to the filter script.
 				for (int PathRemovalCoiunt = 0; PathRemovalCoiunt < 5; ++PathRemovalCoiunt)
 				{
+					// Back up five levels.
 					FilterPath = FilterPath.Remove(FilterPath.LastIndexOf('/'));
 				}
-				FilterPath += "/Source/Programs/UnrealDocTool/DoxygenInputFilter/CommentFixup.py";
+				FilterPath += "/Source/Programs/UnrealDocTool/DoxygenInputFilter/bin/Release/DoxygenInputFilter.exe";
 				FormatSetting(Output, "INPUT_FILTER", FilterPath);
-				//FormatSetting(Output, "INPUT_FILTER", "D:/Builds_ArtistSync/UE4/Engine/Source/Programs/UnrealDocTool/DoxygenInputFilter/CommentFixup.py");
 
 				// List of predefined macros
 				Output.WriteLine();
