Class descriptions now tolerate < and > symbols. Considering making this global, but it will interfere with our ability to use HTML tags. We may not want that ability, however.

diff --git a/Engine/Source/Programs/UnrealDocTool/APIDocTool/APIDocTool/APIRecord.cs b/Engine/Source/Programs/UnrealDocTool/APIDocTool/APIDocTool/APIRecord.cs
--- a/Engine/Source/Programs/UnrealDocTool/APIDocTool/APIDocTool/APIRecord.cs
+++ b/Engine/Source/Programs/UnrealDocTool/APIDocTool/APIDocTool/APIRecord.cs
@@ -367,7 +367,7 @@ namespace APIDocTool
 				if (!Utility.IsNullOrWhitespace(FullDescription))
 				{
 					Writer.EnterSection("description", "Remarks");
-					Writer.WriteLine(FullDescription);
+					Writer.WriteLine(FullDescription.Replace("<", "&lt;").Replace(">", "&gt;"));
 					Writer.LeaveSection();
 				}
 
