TranslatedWordsCountEstimator: - Fix LineEnding in regex on windows.
- Fix a bug that introduced by last changelist that handle hack 0 number changelist mark as untranslated.
- Fix a crash bug that some special case the recorded locfile time is earlier than the first intfile.

diff --git a/Engine/Source/Programs/UnrealDocTool/UnrealDocTranslator/TranslatedWordsCountEstimator/Program.cs b/Engine/Source/Programs/UnrealDocTool/UnrealDocTranslator/TranslatedWordsCountEstimator/Program.cs
--- a/Engine/Source/Programs/UnrealDocTool/UnrealDocTranslator/TranslatedWordsCountEstimator/Program.cs
+++ b/Engine/Source/Programs/UnrealDocTool/UnrealDocTranslator/TranslatedWordsCountEstimator/Program.cs
@@ -13,7 +13,7 @@ namespace TranslatedWordsCountEstimator
 
     class Program
     {
-        private static readonly Regex ChangelistNumberPattern = new Regex(@"^INTSourceChangelist\:(?<number>\d+)\n", RegexOptions.Multiline | RegexOptions.Compiled);
+        private static readonly Regex ChangelistNumberPattern = new Regex(@"^INTSourceChangelist\:(?<number>\d+)"+Environment.NewLine, RegexOptions.Multiline | RegexOptions.Compiled);
 
         private const string DefaultDepotPath = "//depot/UE4";
         private const string DefaultDocumentationPath = "/Engine/Documentation/Source/....udn";
@@ -135,6 +135,11 @@ namespace TranslatedWordsCountEstimator
                         if (match.Success)
                         {
                             var cl = int.Parse(match.Groups["number"].Value);
+                            if (cl == 0)
+                            {
+                                wordCount -= intWordCount.Get(intFile);
+                                continue;
+                            }
                             if (intFile.HeadChange > cl)
                             {
                                 diffDate = p4Repository.GetChangelist(cl).ModifiedDate;
@@ -160,10 +165,15 @@ namespace TranslatedWordsCountEstimator
 
                         if (diffs != null && diffs[0] != null && string.IsNullOrEmpty(diffs[0].LeftFile.DepotPath.Path))
                         {
+                            PathSpec P4Loc = TraceP4Location(p4Repository, new FileSpec(intFile.DepotPath, new ChangelistIdVersion(latestChangelist.Id)), diffDate.Value);
+                            // some wired case which the ModifiedDate for intFile from locFile's Changelist pattern is earlier than the first version Date.
+                            // maybe caused by server time adjustment.
+                            if(P4Loc == null)
+                            {
+                                continue;
+                            }
                             diffs = p4Repository.GetDepotFileDiffs(
-                                string.Format("{0}@{1}", TraceP4Location(p4Repository,
-                                    new FileSpec(intFile.DepotPath, new ChangelistIdVersion(latestChangelist.Id)), diffDate.Value),
-                                    diffDate.Value.ToString("yyyy\\/MM\\/dd\\:HH\\:mm\\:ss")),
+                                string.Format("{0}@{1}", P4Loc, diffDate.Value.ToString("yyyy\\/MM\\/dd\\:HH\\:mm\\:ss")),
                                 new FileSpec(intFile.DepotPath, new ChangelistIdVersion(latestChangelist.Id)).ToString(),
                                 new Options() { { "-d", "" }, { "-u", "" } });
                         }
