#UnrealDocTranslator: fix a bug that newly translated file doesn't check out automatically and set the correct changelist nubmer as metadata.

diff --git a/Engine/Source/Programs/UnrealDocTool/UnrealDocTranslator/UnrealDocTranslator/UnrealDocTranslatorForm.cs b/Engine/Source/Programs/UnrealDocTool/UnrealDocTranslator/UnrealDocTranslator/UnrealDocTranslatorForm.cs
--- a/Engine/Source/Programs/UnrealDocTool/UnrealDocTranslator/UnrealDocTranslator/UnrealDocTranslatorForm.cs
+++ b/Engine/Source/Programs/UnrealDocTool/UnrealDocTranslator/UnrealDocTranslator/UnrealDocTranslatorForm.cs
@@ -1052,7 +1052,7 @@ namespace UnrealDocTranslator
 
                                 try
                                 {
-                                    P4Connection.Client.AddFiles(CurrentClientFile, null);
+                                    P4Connection.Client.EditFiles(CurrentClientFile, null);
                                     SelectedRow.ForeColor = Color.Red;
                                     FilesCheckedOut.Add(GetLangFileNameFromInt(GetDepotFile(SelectedRow)));
                                 }
@@ -1064,6 +1064,12 @@ namespace UnrealDocTranslator
                                 }
                             }
                         }
+
+                        // As a new file, we should handle the INTSourceChangelist number correctly here
+                        string NewFileContents = System.IO.File.ReadAllText(LangClientFileName);
+                        NewFileContents = ParseSourceChangelist.Replace(NewFileContents, "INTSourceChangelist:" + GetChangelist(SelectedRow));
+                        System.IO.File.WriteAllText(LangClientFileName, NewFileContents, Encoding.UTF8);
+
                         string AraxisMergeComparePath = System.IO.Path.Combine(_connectionDetails.AraxisFolderName, "Compare.exe");
                         if ((new System.IO.FileInfo(AraxisMergeComparePath)).Exists)
                         {
