[NO NON-PII COMMIT MESSAGE]

diff --git a/Engine/Source/Programs/UnrealDocTool/APIDocTool/APIDocTool/Program.cs b/Engine/Source/Programs/UnrealDocTool/APIDocTool/APIDocTool/Program.cs
--- a/Engine/Source/Programs/UnrealDocTool/APIDocTool/APIDocTool/Program.cs
+++ b/Engine/Source/Programs/UnrealDocTool/APIDocTool/APIDocTool/Program.cs
@@ -645,7 +645,7 @@ namespace APIDocTool
 				{
 					return false;
 				}
-				File.Copy(Path.Combine(EngineDir, "Intermediate\\Build\\UBTExport.0.xge.xml"), TargetInfoPath, true);
+				File.Copy(Path.Combine(EngineDir, "Intermediate\\Build\\UBTExport.000.xge.xml"), TargetInfoPath, true);
 			}
 			if((Actions & BuildActions.Archive) != 0)
 			{
