Fixing some bad string logic in APIDocTool.

diff --git a/Engine/Source/Programs/UnrealDocTool/APIDocTool/APIDocTool/SourceFile.cs b/Engine/Source/Programs/UnrealDocTool/APIDocTool/APIDocTool/SourceFile.cs
--- a/Engine/Source/Programs/UnrealDocTool/APIDocTool/APIDocTool/SourceFile.cs
+++ b/Engine/Source/Programs/UnrealDocTool/APIDocTool/APIDocTool/SourceFile.cs
@@ -122,27 +122,27 @@ namespace APIDocTool
 			List<SourceToken> Tokens = new List<SourceToken>();
 			for(int Position = 0; Position < Text.Length; )
 			{
-				if(Char.IsWhiteSpace(Text[Position]))
+				if (Char.IsWhiteSpace(Text[Position]))
 				{
 					// Skip whitespace
-					Position++;
+					++Position;
 				}
-				else if(Text[Position] == '/' && Text[Position + 1] == '/')
+				else if (((Position + 1) < Text.Length) && (Text[Position] == '/' && Text[Position + 1] == '/'))
 				{
 					// Skip C++ style comment
 					Position += 2;
-					while(Position < Text.Length && Text[Position] != '\n' && Text[Position] != '\r')
+					while ((Position < Text.Length) && (Text[Position] != '\n' && Text[Position] != '\r'))
 					{
-						Position++;
+						++Position;
 					}
 				}
-				else if(Text[Position] == '/' && Text[Position + 1] == '*')
+				else if (((Position + 1) < Text.Length) && (Text[Position] == '/' && Text[Position + 1] == '*'))
 				{
 					// Skip C style comment
 					Position += 4;
-					while(Text[Position - 2] != '*' || Text[Position - 1] != '/')
+					while (((Position - 1) < Text.Length) && (Text[Position - 2] != '*' || Text[Position - 1] != '/'))
 					{
-						Position++;
+						++Position;
 					}
 				}
 				else
@@ -152,7 +152,7 @@ namespace APIDocTool
 					{
 						// It's a character or string literal. 
 						char Terminator = Text[Position++];
-						while (Text[Position] != Terminator)
+						while ((Position < Text.Length) && (Text[Position] != Terminator))
 						{
 							if (Text[Position] == '\\')
 							{
