[NO NON-PII COMMIT MESSAGE]

diff --git a/Engine/Source/Programs/UnrealDocTool/APIDocTool/APIDocTool/Program.cs b/Engine/Source/Programs/UnrealDocTool/APIDocTool/APIDocTool/Program.cs
--- a/Engine/Source/Programs/UnrealDocTool/APIDocTool/APIDocTool/Program.cs
+++ b/Engine/Source/Programs/UnrealDocTool/APIDocTool/APIDocTool/Program.cs
@@ -622,7 +622,7 @@ namespace APIDocTool
 
 		static bool BuildCodeTargetInfo(string TargetInfoPath, string EngineDir, string ArchivePath, BuildActions Actions)
 		{
-			if((Actions & BuildActions.Clean) != 0)
+			if ((Actions & BuildActions.Clean) != 0)
 			{
 				Console.WriteLine("Cleaning '{0}'", TargetInfoPath);
 				Utility.SafeDeleteFile(TargetInfoPath);
@@ -645,9 +645,22 @@ namespace APIDocTool
 				{
 					return false;
 				}
-				File.Copy(Path.Combine(EngineDir, "Intermediate\\Build\\UBTExport.000.xge.xml"), TargetInfoPath, true);
+				// Handle different versions of DocumentationEditor on the build system.
+				if (File.Exists(Path.Combine(EngineDir, "Intermediate\\Build\\UBTExport.0.xge.xml")))
+				{
+					File.Copy(Path.Combine(EngineDir, "Intermediate\\Build\\UBTExport.0.xge.xml"), TargetInfoPath, true);
+				}
+				else if (File.Exists(Path.Combine(EngineDir, "Intermediate\\Build\\UBTExport.000.xge.xml")))
+				{
+					File.Copy(Path.Combine(EngineDir, "Intermediate\\Build\\UBTExport.000.xge.xml"), TargetInfoPath, true);
+				}
+				else
+				{
+					Console.WriteLine("Failed to find either UBTExport.0.xge.xml or UBTExport.000.xge.xml under %s", Path.Combine(EngineDir, "Intermediate\\Build\\"));
+					return false;
+				}
 			}
-			if((Actions & BuildActions.Archive) != 0)
+			if ((Actions & BuildActions.Archive) != 0)
 			{
 				Console.WriteLine("Creating archive '{0}'", ArchivePath);
 				Utility.CreateTgzFromDir(ArchivePath, Path.GetDirectoryName(TargetInfoPath));
