[NO NON-PII COMMIT MESSAGE]

diff --git a/Engine/Source/Programs/UnrealDocTool/APIDocTool/APIDocTool/SourceFile.cs b/Engine/Source/Programs/UnrealDocTool/APIDocTool/APIDocTool/SourceFile.cs
--- a/Engine/Source/Programs/UnrealDocTool/APIDocTool/APIDocTool/SourceFile.cs
+++ b/Engine/Source/Programs/UnrealDocTool/APIDocTool/APIDocTool/SourceFile.cs
@@ -120,7 +120,7 @@ namespace APIDocTool
 
 			// Build a list of tokens
 			List<SourceToken> Tokens = new List<SourceToken>();
-			for(int Position = 0; Position < Text.Length; )
+			for (int Position = 0; Position < Text.Length; )
 			{
 				if (Char.IsWhiteSpace(Text[Position]))
 				{
@@ -130,19 +130,17 @@ namespace APIDocTool
 				else if (((Position + 1) < Text.Length) && (Text[Position] == '/' && Text[Position + 1] == '/'))
 				{
 					// Skip C++ style comment
-					Position += 2;
-					while ((Position < Text.Length) && (Text[Position] != '\n' && Text[Position] != '\r'))
+					for (Position += 2; ((Position < Text.Length) && (Text[Position] != '\n' && Text[Position] != '\r')); ++Position)
 					{
-						++Position;
+						// Nothing to do but iterate Position.
 					}
 				}
 				else if (((Position + 1) < Text.Length) && (Text[Position] == '/' && Text[Position + 1] == '*'))
 				{
 					// Skip C style comment
-					Position += 4;
-					while (((Position - 1) < Text.Length) && (Text[Position - 2] != '*' || Text[Position - 1] != '/'))
+					for (Position += 4; (((Position - 1) < Text.Length) && (Text[Position - 2] != '*' || Text[Position - 1] != '/')); ++Position)
 					{
-						++Position;
+						// Nothing to do but iterate Position.
 					}
 				}
 				else
@@ -160,24 +158,24 @@ namespace APIDocTool
 							}
 							else
 							{
-								Position++;
+								++Position;
 							}
 						}
-						Position++;
+						++Position;
 					}
-					else if(Char.IsDigit(Text[Position]) || (Text[Position] == '.' && Char.IsDigit(Text[Position + 1])))
+					else if (Char.IsDigit(Text[Position]) || (((Position + 1) < Text.Length) && (Text[Position] == '.' && Char.IsDigit(Text[Position + 1]))))
 					{
 						// It's a numeric token
-						Position++;
-						for(;;)
+						++Position;
+						while (true)
 						{
-							if((Text[Position] == 'e' || Text[Position] == 'E') && (Text[Position + 1] == '+' || Text[Position + 1] == '-'))
+							if (((Position + 1) < Text.Length) && ((Text[Position] == 'e' || Text[Position] == 'E') && (Text[Position + 1] == '+' || Text[Position + 1] == '-')))
 							{
 								Position += 2;
 							}
-							else if(Char.IsLetterOrDigit(Text[Position]) || Text[Position] == '.')
+							else if ((Position < Text.Length) && (Char.IsLetterOrDigit(Text[Position]) || Text[Position] == '.'))
 							{
-								Position++;
+								++Position;
 							}
 							else
 							{
@@ -185,13 +183,12 @@ namespace APIDocTool
 							}
 						}
 					}
-					else if(Char.IsLetter(Text[Position]) || Text[Position] == '_')
+					else if (Char.IsLetter(Text[Position]) || Text[Position] == '_')
 					{
 						// It's an identifier
-						Position++;
-						while(Position < Text.Length && (Char.IsLetterOrDigit(Text[Position]) || Text[Position] == '_'))
+						for (++Position; ((Position < Text.Length) && (Char.IsLetterOrDigit(Text[Position]) || Text[Position] == '_')); ++Position)
 						{
-							Position++;
+							// Nothing to do but iterate Position.
 						}
 					}
 					else
@@ -199,7 +196,7 @@ namespace APIDocTool
 						// Just take a single character, we don't really care what symbol it is
 						Position++;
 					}
-					Tokens.Add(new SourceToken(this, StartPosition, Position - StartPosition));
+					Tokens.Add(new SourceToken(this, StartPosition, Math.Min(Position, Text.Length) - StartPosition));
 				}
 			}
 			this.Tokens = Tokens.ToArray();
