[NO NON-PII COMMIT MESSAGE]

diff --git a/Engine/Source/Programs/UnrealDocTool/APIDocTool/APIDocTool/Utility.cs b/Engine/Source/Programs/UnrealDocTool/APIDocTool/APIDocTool/Utility.cs
--- a/Engine/Source/Programs/UnrealDocTool/APIDocTool/APIDocTool/Utility.cs
+++ b/Engine/Source/Programs/UnrealDocTool/APIDocTool/APIDocTool/Utility.cs
@@ -482,6 +482,11 @@ namespace APIDocTool
 
 		public static void CreateTgzFromFiles(string TarPath, string BaseDir, IEnumerable<string> InputFiles)
 		{
+			// File creation does not implicitly create the directory for the file, so check/do that part first.
+			if (!Directory.Exists(Path.GetDirectoryName(TarPath)))
+			{
+				Directory.CreateDirectory(Path.GetDirectoryName(TarPath));
+			}
 			using (FileStream TarFileStream = new FileStream(TarPath, FileMode.Create))
 			{
 				using (GZipStream ZipStream = new GZipStream(TarFileStream, CompressionMode.Compress))
