Copying APIDocTool ini file fix from 4.11 branch to expedite stopping the CIS failures in main.

diff --git a/Engine/Source/Programs/UnrealDocTool/APIDocTool/APIDocTool/IniFile.cs b/Engine/Source/Programs/UnrealDocTool/APIDocTool/APIDocTool/IniFile.cs
--- a/Engine/Source/Programs/UnrealDocTool/APIDocTool/APIDocTool/IniFile.cs
+++ b/Engine/Source/Programs/UnrealDocTool/APIDocTool/APIDocTool/IniFile.cs
@@ -12,7 +12,7 @@ namespace APIDocTool
 	public class IniSection
 	{
 		public string Name;
-		public Dictionary<string, string> Values = new Dictionary<string,string>();
+		public Dictionary<string, string> Values = new Dictionary<string, string>();
 
 		public IniSection(string InName)
 		{
@@ -67,7 +67,7 @@ namespace APIDocTool
 				if (Section != null)
 				{
 					string Value = Section.Find(KeyPath.Substring(EndIndex + 1));
-					if(Value != null) return Value;
+					if (Value != null) return Value;
 				}
 			}
 			return null;
@@ -76,7 +76,7 @@ namespace APIDocTool
 		public string FindValueOrDefault(string KeyPath, string DefaultValue)
 		{
 			string Value = FindValue(KeyPath);
-			return (Value != null)? Value : DefaultValue;
+			return (Value != null) ? Value : DefaultValue;
 		}
 
 		public IniSection FindSection(string Name)
@@ -94,7 +94,13 @@ namespace APIDocTool
 
 		public static IniFile Read(string InputPath)
 		{
-			IniFile File = new IniFile();
+			if (!File.Exists(InputPath))
+			{
+				Console.WriteLine("Ini File %s not found!", InputPath);
+				return new IniFile();
+			}
+
+			IniFile InputIniFile = new IniFile();
 			using (StreamReader Reader = new StreamReader(InputPath))
 			{
 				IniSection CurrentSection = null;
@@ -111,12 +117,12 @@ namespace APIDocTool
 							string Name = Line.Substring(1, Line.Length - 2);
 							CurrentSection = new IniSection(Name);
 
-							File.Sections.Add(CurrentSection);
-							if(File.NameToSection.ContainsKey(Name.ToLowerInvariant()))
+							InputIniFile.Sections.Add(CurrentSection);
+							if (InputIniFile.NameToSection.ContainsKey(Name.ToLowerInvariant()))
 							{
 								Console.WriteLine("{0} has already been added", Name);
 							}
-							File.NameToSection.Add(Name.ToLowerInvariant(), CurrentSection);
+							InputIniFile.NameToSection.Add(Name.ToLowerInvariant(), CurrentSection);
 						}
 						else
 						{
@@ -124,7 +130,7 @@ namespace APIDocTool
 
 							// Split the line into key/value
 							int ValueIdx = Line.IndexOf('=');
-							if(ValueIdx == -1)
+							if (ValueIdx == -1)
 							{
 								Key = Line;
 								Value = "";
@@ -136,7 +142,7 @@ namespace APIDocTool
 							}
 
 							// If it's an append operation, add the existing value
-							if(Key.Length > 0 && Key[0] == '+')
+							if (Key.Length > 0 && Key[0] == '+')
 							{
 								CurrentSection.Append(Key.Substring(1), Value);
 							}
@@ -148,7 +154,7 @@ namespace APIDocTool
 					}
 				}
 			}
-			return File;
+			return InputIniFile;
 		}
 	}
 }
