Fix code documentation build failures; added support for parsing response files to APIDocTool required after Dev-Platform integration.

diff --git a/Engine/Source/Programs/UnrealDocTool/APIDocTool/APIDocTool/BuildTarget.cs b/Engine/Source/Programs/UnrealDocTool/APIDocTool/APIDocTool/BuildTarget.cs
--- a/Engine/Source/Programs/UnrealDocTool/APIDocTool/APIDocTool/BuildTarget.cs
+++ b/Engine/Source/Programs/UnrealDocTool/APIDocTool/APIDocTool/BuildTarget.cs
@@ -31,25 +31,16 @@ namespace APIDocTool
 
 					// Parse a list of tokens
 					List<string> Tokens = new List<string>();
-					for (int Idx = 0; Idx < Params.Length; Idx++)
+					Tokens.AddRange(TokenizeArguments(Params));
+
+					// Expand response files
+					for(int Idx = 0; Idx < Tokens.Count; Idx++)
 					{
-						if (!Char.IsWhiteSpace(Params[Idx]))
+						if(Tokens[Idx].StartsWith("@"))
 						{
-							StringBuilder Token = new StringBuilder();
-							while (Idx < Params.Length && !Char.IsWhiteSpace(Params[Idx]))
-							{
-								if (Params[Idx] == '"')
-								{
-									Idx++;
-									while (Params[Idx] != '"') Token.Append(Params[Idx++]);
-									Idx++;
-								}
-								else
-								{
-									Token.Append(Params[Idx++]);
-								}
-							}
-							Tokens.Add(Token.ToString());
+							List<string> NewTokens = File.ReadAllLines(Tokens[Idx].Substring(1)).SelectMany(Line => TokenizeArguments(Line)).ToList();
+							Tokens.RemoveAt(Idx);
+							Tokens.InsertRange(Idx, NewTokens);
 						}
 					}
 
@@ -78,6 +69,31 @@ namespace APIDocTool
 			}
 		}
 
+		static IEnumerable<string> TokenizeArguments(string Params)
+		{
+			for (int Idx = 0; Idx < Params.Length; Idx++)
+			{
+				if (!Char.IsWhiteSpace(Params[Idx]))
+				{
+					StringBuilder Token = new StringBuilder();
+					while (Idx < Params.Length && !Char.IsWhiteSpace(Params[Idx]))
+					{
+						if (Params[Idx] == '"')
+						{
+							Idx++;
+							while (Params[Idx] != '"') Token.Append(Params[Idx++]);
+							Idx++;
+						}
+						else
+						{
+							Token.Append(Params[Idx++]);
+						}
+					}
+					yield return Token.ToString();
+				}
+			}
+		}
+
 		static string ResolveFile(string BaseDirectoryName, string FileName)
 		{
 			string NewFileName = FileName;
