APIDocTool: Ensure that Doxygen crashing produces an error message that's picked up by EC. Currently fails with a warning.

diff --git a/Engine/Source/Programs/UnrealDocTool/APIDocTool/APIDocTool/Program.cs b/Engine/Source/Programs/UnrealDocTool/APIDocTool/APIDocTool/Program.cs
--- a/Engine/Source/Programs/UnrealDocTool/APIDocTool/APIDocTool/Program.cs
+++ b/Engine/Source/Programs/UnrealDocTool/APIDocTool/APIDocTool/Program.cs
@@ -1075,7 +1075,7 @@ namespace APIDocTool
                     // Run Doxygen
                     if (!Doxygen.Run(DoxygenPath, Path.Combine(EngineDir, "Source"), Config, true))
                     {
-                        Console.WriteLine("  Doxygen crashed. Skipping.");
+                        Console.WriteLine("  error: Doxygen crashed. Skipping.");
                         return false;
                     }
                 }
