[NO NON-PII COMMIT MESSAGE]

diff --git a/Engine/Source/Programs/UnrealDocTool/DoxygenInputFilter/DoxygenInputFilter.pl b/Engine/Source/Programs/UnrealDocTool/DoxygenInputFilter/DoxygenInputFilter.pl
--- a/Engine/Source/Programs/UnrealDocTool/DoxygenInputFilter/DoxygenInputFilter.pl
+++ b/Engine/Source/Programs/UnrealDocTool/DoxygenInputFilter/DoxygenInputFilter.pl
@@ -5,42 +5,64 @@
 # If it's immediately followed by "~" (our "don't comment" code), ignore it.
 # If it's a "///" or "/**" (Doxygen's "comment" code), ignore it.
 # Otherwise, transform it into the appropriate Doxygen "comment" code.
+# For .cpp files, ignore all comment-modification code as it's causing undesired comments (e.g. copyright notices) to be captured.
+# Turn enum classes into regular enums, also remove the storage type specifier if it's there. This is due to our use of an old, customized Doxygen.
 
 my $line = "";
 my $commentposition;
 my $tempvalue;
-open(FILE1, $ARGV[0]) || die "Error: $!\n";
+open(FILE1, $ARGV[0]) || die "  error: $!\n";
 while (<FILE1>)
 {
     $line = $_;
+	
+	if (lc(substr($ARGV[0], -4)) ne ".cpp")
+	{
+		$commentposition = index($line, "\/\/");
+		$tempvalue = index($line, "\/*");
+		if ($commentposition < 0)
+		{
+			$commentposition = $tempvalue;
+		}
+		elsif ($tempvalue >= 0)
+		{
+			$commentposition = $commentposition < $tempvalue ? $commentposition : $tempvalue;
+		}
+		
+		if ($commentposition >= 0)
+		{
+			if ((length $line) > ($commentposition + 2))
+			{
+				$tempvalue = substr($line, $commentposition + 2, 1);
+				if (!($tempvalue eq "~") && !($tempvalue eq substr($line, $commentposition + 1, 1)))
+				{
+					substr($line, $commentposition + 1, 0) = substr($line, $commentposition + 1, 1);
+				}
+			}
+			# This elsif indicates an error - the line didn't end with "\n".
+			# If we find this happening, we should start appending "\n" on our own.
+			#elsif ((length $line) > ($commentposition + 1))
+			#{
+			#    substr($line, $commentposition + 1, 0) = substr($line, $commentposition + 1, 1);
+			#}
+		}
+	}
+	print fixenums($line);
+}
 
-    $commentposition = index($line, "\/\/");
-    $tempvalue = index($line, "\/*");
-    if ($commentposition < 0)
-    {
-        $commentposition = $tempvalue;
-    }
-    elsif ($tempvalue >= 0)
-    {
-        $commentposition = $commentposition < $tempvalue ? $commentposition : $tempvalue;
-    }
-    
-    if ($commentposition >= 0)
-    {
-        if ((length $line) > ($commentposition + 2))
-        {
-            $tempvalue = substr($line, $commentposition + 2, 1);
-            if (!($tempvalue eq "~") && !($tempvalue eq substr($line, $commentposition + 1, 1)))
-            {
-                substr($line, $commentposition + 1, 0) = substr($line, $commentposition + 1, 1);
-            }
-        }
-        # This elsif indicates an error - the line didn't end with "\n".
-        # If we find this happening, we should start appending "\n" on our own.
-        #elsif ((length $line) > ($commentposition + 1))
-        #{
-        #    substr($line, $commentposition + 1, 0) = substr($line, $commentposition + 1, 1);
-        #}
-    }
-    print $line;
+# The version of Doxygen we have does not recognize "enum class", but we now use them a lot. Rephrasing these as regular enums seems to fix the issue.
+# Since enum classes are stricter than enums, changing to the looser type should never cause a problem, and Doxygen should never know the difference.
+# We are also dropping the storage type specification from the end, as Doxygen seems not to understand that, either.
+sub fixenums
+{
+	my $newline = $_[0];
+	if ($newline =~ m/\s*enum\s+class\s+/)
+	{
+		$newline =~ s/\s*enum\s+class\s+/enum /;
+	}
+	if ($newline =~ m/.*:(\s)*(u)?int(8|16|32)/)
+	{
+		$newline =~ s/:(\s)*(u)?int(8|16|32)/ /;
+	}
+	return $newline;
 }
