APIDocTool mods to go with new Doxygen. Currently set to run the Debugx64 Doxygen executable while Releasex64 executable is being tested.

diff --git a/Engine/Source/Programs/UnrealDocTool/APIDocTool/APIDocTool/Program.cs b/Engine/Source/Programs/UnrealDocTool/APIDocTool/APIDocTool/Program.cs
--- a/Engine/Source/Programs/UnrealDocTool/APIDocTool/APIDocTool/Program.cs
+++ b/Engine/Source/Programs/UnrealDocTool/APIDocTool/APIDocTool/Program.cs
@@ -462,7 +462,7 @@ namespace APIDocTool
 			}
 
 			// Derive all the tool paths
-			string DoxygenPath = Path.Combine(EngineDir, "Extras\\NotForLicensees\\Doxygen\\bin\\Release64\\doxygen.exe");
+			string DoxygenPath = Path.Combine(EngineDir, "Extras\\NotForLicensees\\doxygen-1.8.11\\Doxygen\\Debugx64\\doxygen.exe");
 			string EditorPath = Path.Combine(EngineDir, "Binaries\\Win64\\UE4Editor-Cmd.exe");
 			string DocToolPath = Path.Combine(EngineDir, "Binaries\\DotNET\\UnrealDocTool.exe");
 			string ChmCompilerPath = Path.Combine(EngineDir, "Extras\\NotForLicensees\\HTML Help Workshop\\hhc.exe");
@@ -1464,7 +1464,7 @@ namespace APIDocTool
 				{
 					DocToolProcess.StartInfo.WorkingDirectory = EngineDir;
 					DocToolProcess.StartInfo.FileName = DocToolPath;
-					DocToolProcess.StartInfo.Arguments = Path.Combine(HtmlSuffix, "*") + " -lang=INT -t=DefaultAPI.html -v=warn";
+					DocToolProcess.StartInfo.Arguments = Path.Combine(HtmlSuffix, "*") + " -lang=INT -t=DefaultAPI.html -v=warn -nodb=true";
 					DocToolProcess.StartInfo.UseShellExecute = false;
 					DocToolProcess.StartInfo.RedirectStandardOutput = true;
 					DocToolProcess.StartInfo.RedirectStandardError = true;
@@ -1541,7 +1541,8 @@ namespace APIDocTool
 
 		static private void CreateChm(string ChmCompilerPath, string ChmFileName, string Title, string DefaultTopicPath, string ContentsFileName, string IndexFileName, string SourceDir, List<string> FileNames)
 		{
-			string ProjectName = Path.GetFileNameWithoutExtension(ChmFileName);
+            Console.WriteLine("CreateChm {0}", ChmFileName);
+            string ProjectName = Path.GetFileNameWithoutExtension(ChmFileName);
 
 			// Create an intermediate directory
 			string IntermediateDir = Path.Combine(Path.GetDirectoryName(ChmFileName), ProjectName);
diff --git a/Engine/Source/Programs/UnrealDocTool/DoxygenInputFilter/DoxygenInputFilter.pl b/Engine/Source/Programs/UnrealDocTool/DoxygenInputFilter/DoxygenInputFilter.pl
--- a/Engine/Source/Programs/UnrealDocTool/DoxygenInputFilter/DoxygenInputFilter.pl
+++ b/Engine/Source/Programs/UnrealDocTool/DoxygenInputFilter/DoxygenInputFilter.pl
@@ -6,7 +6,6 @@
 # If it's a "///" or "/**" (Doxygen's "comment" code), ignore it.
 # Otherwise, transform it into the appropriate Doxygen "comment" code.
 # For .cpp files, ignore all comment-modification code as it's causing undesired comments (e.g. copyright notices) to be captured.
-# Turn enum classes into regular enums, also remove the storage type specifier if it's there. This is due to our use of an old, customized Doxygen.
 
 my $line = "";
 my $commentposition;
@@ -47,22 +46,5 @@ while (<FILE1>)
 			#}
 		}
 	}
-	print fixenums($line);
-}
-
-# The version of Doxygen we have does not recognize "enum class", but we now use them a lot. Rephrasing these as regular enums seems to fix the issue.
-# Since enum classes are stricter than enums, changing to the looser type should never cause a problem, and Doxygen should never know the difference.
-# We are also dropping the storage type specification from the end, as Doxygen seems not to understand that, either.
-sub fixenums
-{
-	my $newline = $_[0];
-	if ($newline =~ m/\s*enum\s+class\s+/)
-	{
-		$newline =~ s/\s*enum\s+class\s+/enum /;
-	}
-	if ($newline =~ m/.*:(\s)*(u)?int(8|16|32)/)
-	{
-		$newline =~ s/:(\s)*(u)?int(8|16|32)/ /;
-	}
-	return $newline;
+	print $line;
 }
diff --git a/Engine/Source/Programs/UnrealDocTool/DoxygenLib/DoxygenModule.cs b/Engine/Source/Programs/UnrealDocTool/DoxygenLib/DoxygenModule.cs
--- a/Engine/Source/Programs/UnrealDocTool/DoxygenLib/DoxygenModule.cs
+++ b/Engine/Source/Programs/UnrealDocTool/DoxygenLib/DoxygenModule.cs
@@ -70,6 +70,12 @@ namespace DoxygenLib
 			{
 				Entity.File = LocationNode.Attributes["file"].Value.Replace('/', '\\');
 
+                // Fix Doxygen's inconsistent pathing.
+                if (!Entity.File.StartsWith(InModule.BaseSrcDir))
+                {
+                    Entity.File = InModule.BaseSrcDir + "Source\\" + Entity.File;
+                }
+
 				XmlAttribute LineAttribute = LocationNode.Attributes["line"];
 				if (LineAttribute != null)
 				{
@@ -83,7 +89,12 @@ namespace DoxygenLib
 				if (BodyFileAttribute != null && BodyStartAttribute != null && BodyEndAttribute != null)
 				{
 					Entity.BodyFile = BodyFileAttribute.Value.Replace('/', '\\');
-					Entity.BodyStart = Int32.Parse(BodyStartAttribute.Value);
+                    // Fix Doxygen's inconsistent pathing.
+                    if (!Entity.BodyFile.StartsWith(InModule.BaseSrcDir))
+                    {
+                        Entity.BodyFile = InModule.BaseSrcDir + "Source\\" + Entity.BodyFile;
+                    }
+                    Entity.BodyStart = Int32.Parse(BodyStartAttribute.Value);
 					Entity.BodyEnd = Int32.Parse(BodyEndAttribute.Value);
 				}
 			}
@@ -127,7 +138,7 @@ namespace DoxygenLib
 			return Module;
 		}
 
-		public static DoxygenModule TryRead(string Name, string BaseSrcDir, string BaseXmlDir)
+        public static DoxygenModule TryRead(string Name, string BaseSrcDir, string BaseXmlDir)
 		{
 			// Load the index
 			XmlDocument Document;
