Changing over to Doxygen Releasex64 build following successful test. Removing defunct "-nodb" argument from UnrealDocTool call - need further investigation to see what that argument was intended to do and why it didn't work.

diff --git a/Engine/Source/Programs/UnrealDocTool/APIDocTool/APIDocTool/Program.cs b/Engine/Source/Programs/UnrealDocTool/APIDocTool/APIDocTool/Program.cs
--- a/Engine/Source/Programs/UnrealDocTool/APIDocTool/APIDocTool/Program.cs
+++ b/Engine/Source/Programs/UnrealDocTool/APIDocTool/APIDocTool/Program.cs
@@ -462,7 +462,7 @@ namespace APIDocTool
 			}
 
 			// Derive all the tool paths
-			string DoxygenPath = Path.Combine(EngineDir, "Extras\\NotForLicensees\\doxygen-1.8.11\\Doxygen\\Debugx64\\doxygen.exe");
+			string DoxygenPath = Path.Combine(EngineDir, "Extras\\NotForLicensees\\doxygen-1.8.11\\Doxygen\\Releasex64\\doxygen.exe");
 			string EditorPath = Path.Combine(EngineDir, "Binaries\\Win64\\UE4Editor-Cmd.exe");
 			string DocToolPath = Path.Combine(EngineDir, "Binaries\\DotNET\\UnrealDocTool.exe");
 			string ChmCompilerPath = Path.Combine(EngineDir, "Extras\\NotForLicensees\\HTML Help Workshop\\hhc.exe");
@@ -497,7 +497,7 @@ namespace APIDocTool
 			// Establish snippet directory so we can look things up later
 			APISnippets.SetSnippetTextDirectory(SnippetsDir);
 
-			// Build all the code docs
+            // Build all the code docs
 			if (!BuildCodeSnippetsTxt(SamplesDir, CodeSnippetsActions))
 			{
 				return 1;
@@ -1464,7 +1464,7 @@ namespace APIDocTool
 				{
 					DocToolProcess.StartInfo.WorkingDirectory = EngineDir;
 					DocToolProcess.StartInfo.FileName = DocToolPath;
-					DocToolProcess.StartInfo.Arguments = Path.Combine(HtmlSuffix, "*") + " -lang=INT -t=DefaultAPI.html -v=warn -nodb=true";
+					DocToolProcess.StartInfo.Arguments = Path.Combine(HtmlSuffix, "*") + " -lang=INT -t=DefaultAPI.html -v=warn";
 					DocToolProcess.StartInfo.UseShellExecute = false;
 					DocToolProcess.StartInfo.RedirectStandardOutput = true;
 					DocToolProcess.StartInfo.RedirectStandardError = true;
