Improving filtering to remove copyright line (wrongly captured as documentation in some cases) and work around warnings created by the implicit constructors contained in the GENERATED_UCLASS_BODY macro.

diff --git a/Engine/Source/Programs/UnrealDocTool/DoxygenInputFilter/DoxygenInputFilter.pl b/Engine/Source/Programs/UnrealDocTool/DoxygenInputFilter/DoxygenInputFilter.pl
--- a/Engine/Source/Programs/UnrealDocTool/DoxygenInputFilter/DoxygenInputFilter.pl
+++ b/Engine/Source/Programs/UnrealDocTool/DoxygenInputFilter/DoxygenInputFilter.pl
@@ -8,6 +8,7 @@
 # For .cpp files, ignore all comment-modification code as it's causing undesired comments (e.g. copyright notices) to be captured.
 
 my $line = "";
+my $lastclassname;
 my $commentposition;
 my $tempvalue;
 open(FILE1, $ARGV[0]) || die "  error: $!\n";
@@ -15,8 +16,24 @@ while (<FILE1>)
 {
     $line = $_;
 	
-	if (lc(substr($ARGV[0], -4)) ne ".cpp")
+	if ($line =~ /\/\/ Copyright .+ Epic Games, Inc. All Rights Reserved./)
 	{
+		# Discard copyright line.
+		$line = " ";
+	}
+	else
+	{
+		if ($line =~ /class .+ :.*/)
+		{
+			$lastclassname = substr($line, index($line, "class ") + 6, index($line, ":") - 6);
+		}
+		else
+		{
+			if (index($line, "GENERATED_UCLASS_BODY()") >= 0)
+			{
+				$line = "GENERATED_BODY() public: " . $lastclassname . "(const FObjectInitializer& ObjectInitializer);\n";
+			}
+		}
 		$commentposition = index($line, "\/\/");
 		$tempvalue = index($line, "\/*");
 		if ($commentposition < 0)
