Reduced warning count from >2000 to 13. The remaining ones are likely going to be much slower to fix and maybe br Doxygen-related, as opposed to being our engine itself or APIDocTool-related code.

diff --git a/Engine/Source/Programs/UnrealDocTool/APIDocTool/APIDocTool/IniFile.cs b/Engine/Source/Programs/UnrealDocTool/APIDocTool/APIDocTool/IniFile.cs
--- a/Engine/Source/Programs/UnrealDocTool/APIDocTool/APIDocTool/IniFile.cs
+++ b/Engine/Source/Programs/UnrealDocTool/APIDocTool/APIDocTool/IniFile.cs
@@ -96,7 +96,7 @@ namespace APIDocTool
 		{
 			if (!File.Exists(InputPath))
 			{
-				Console.WriteLine("Ini File %s not found!", InputPath);
+				Console.WriteLine("Ini File {0} not found!", InputPath);
 				return new IniFile();
 			}
 
diff --git a/Engine/Source/Programs/UnrealDocTool/APIDocTool/APIDocTool/Program.cs b/Engine/Source/Programs/UnrealDocTool/APIDocTool/APIDocTool/Program.cs
--- a/Engine/Source/Programs/UnrealDocTool/APIDocTool/APIDocTool/Program.cs
+++ b/Engine/Source/Programs/UnrealDocTool/APIDocTool/APIDocTool/Program.cs
@@ -654,7 +654,7 @@ namespace APIDocTool
 				}
 				else
 				{
-					Console.WriteLine("Failed to find either UBTExport.0.xge.xml or UBTExport.000.xge.xml under %s", Path.Combine(EngineDir, "Intermediate\\Build\\"));
+					Console.WriteLine("Failed to find either UBTExport.0.xge.xml or UBTExport.000.xge.xml under {0}", Path.Combine(EngineDir, "Intermediate\\Build\\"));
 					return false;
 				}
 			}
@@ -707,7 +707,7 @@ namespace APIDocTool
 			}
 			if((Actions & BuildActions.Build) != 0)
 			{
-				string MetadataInputPath = Path.Combine(EngineDir, "Source\\Runtime\\CoreUObject\\Public\\UObject\\ObjectBase.h");
+				string MetadataInputPath = Path.Combine(EngineDir, "Source\\Runtime\\CoreUObject\\Public\\UObject\\ObjectMacros.h");
 				Console.WriteLine("Building metadata descriptions from '{0}'...", MetadataInputPath);
 
 				DoxygenConfig Config = new DoxygenConfig("Metadata", new string[]{ MetadataInputPath }, MetadataDir);
diff --git a/Engine/Source/Programs/UnrealDocTool/DoxygenInputFilter/DoxygenInputFilter.pl b/Engine/Source/Programs/UnrealDocTool/DoxygenInputFilter/DoxygenInputFilter.pl
--- a/Engine/Source/Programs/UnrealDocTool/DoxygenInputFilter/DoxygenInputFilter.pl
+++ b/Engine/Source/Programs/UnrealDocTool/DoxygenInputFilter/DoxygenInputFilter.pl
@@ -6,6 +6,7 @@
 # If it's a "///" or "/**" (Doxygen's "comment" code), ignore it.
 # Otherwise, transform it into the appropriate Doxygen "comment" code.
 # For .cpp files, ignore all comment-modification code as it's causing undesired comments (e.g. copyright notices) to be captured.
+# Convert "GENERATED_UCLASS_BODY()" macros to "GENERATED_BODY()", set access level to public, and declare the old-style UE4 constructor, all on the same line.
 
 my $line = "";
 my $lastclassname;
@@ -19,18 +20,43 @@ while (<FILE1>)
 	if ($line =~ /\/\/ Copyright .+ Epic Games, Inc. All Rights Reserved./)
 	{
 		# Discard copyright line.
-		$line = " ";
+		$line = " \n";
 	}
 	else
 	{
-		if ($line =~ /class .+ :.*/)
+		if (index($line, "class ") eq 0)
 		{
-			$lastclassname = substr($line, index($line, "class ") + 6, index($line, ":") - 6);
+			if ($line =~ /class \w+_API \w+/)
+			{
+				$tempvalue = index($line, "_API") + 5;
+				if (index($line, ":") >= 0)
+				{
+					$lastclassname = substr($line, $tempvalue, index($line, ":") - $tempvalue);
+				}
+				else
+				{
+					$lastclassname = substr($line, $tempvalue);
+				}
+#print "lcn0 = " . $lastclassname . "\n";
+			}
+			elsif (index($line, ":") >= 0)
+			{
+				$tempvalue = index($line, "class ") + 6;
+				$lastclassname = substr($line, $tempvalue, index($line, ":") - $tempvalue);
+#print "lcn1 = " . $lastclassname . "\n";
+			}
+			else
+			{
+				# This will catch forward-declared classes, but should not be a problem because these should never immediately precede a GENERATED_UCLASS_BODY() line.
+				$lastclassname = substr($line, index($line, "class ") + 6);
+#print "lcn2 = " . $lastclassname . "\n";
+			}
 		}
 		else
 		{
 			if (index($line, "GENERATED_UCLASS_BODY()") >= 0)
 			{
+				# $lastclassname might not be initialized. If so, the previous logic block failed to catch a class name, or "GENERATED_UCLASS_BODY()" is floating around loose in a file somewhere.
 				$line = "GENERATED_BODY() public: " . $lastclassname . "(const FObjectInitializer& ObjectInitializer);\n";
 			}
 		}
diff --git a/Engine/Source/Programs/UnrealDocTool/DoxygenLib/Doxygen.cs b/Engine/Source/Programs/UnrealDocTool/DoxygenLib/Doxygen.cs
--- a/Engine/Source/Programs/UnrealDocTool/DoxygenLib/Doxygen.cs
+++ b/Engine/Source/Programs/UnrealDocTool/DoxygenLib/Doxygen.cs
@@ -89,7 +89,8 @@ namespace DoxygenLib
 
 				// List of valid file patterns
 				Output.WriteLine();
-				FormatSetting(Output, "FILE_PATTERNS", new List<string> { "*.h", "*.c", "*.hpp", "*.cpp", "*.inl", "*.inc" });
+				//FormatSetting(Output, "FILE_PATTERNS", new List<string> { "*.h", "*.c", "*.hpp", "*.cpp", "*.inl", "*.inc" });
+                FormatSetting(Output, "FILE_PATTERNS", new List<string> { "*.h", "*.hpp", "*.inl", "*.inc" });
 
 				// List of tag files
 				Output.WriteLine();
