Reducing warning count and adding features. Some minor code adjustment will be needed to clear the remaining warnings.

diff --git a/Engine/Source/Programs/UnrealDocTool/DoxygenInputFilter/DoxygenInputFilter.pl b/Engine/Source/Programs/UnrealDocTool/DoxygenInputFilter/DoxygenInputFilter.pl
--- a/Engine/Source/Programs/UnrealDocTool/DoxygenInputFilter/DoxygenInputFilter.pl
+++ b/Engine/Source/Programs/UnrealDocTool/DoxygenInputFilter/DoxygenInputFilter.pl
@@ -7,21 +7,57 @@
 # Otherwise, transform it into the appropriate Doxygen "comment" code.
 # For .cpp files, ignore all comment-modification code as it's causing undesired comments (e.g. copyright notices) to be captured.
 # Convert "GENERATED_UCLASS_BODY()" macros to "GENERATED_BODY()", set access level to public, and declare the old-style UE4 constructor, all on the same line.
+# Remove UENUM, USTRUCT, and UCLASS macros because they confuse Doxygen in some cases.
 
 my $line = "";
 my $lastclassname;
 my $commentposition;
 my $tempvalue;
+my $lcline;
+my $token;
+my $documenton = 1;
 open(FILE1, $ARGV[0]) || die "  error: $!\n";
 while (<FILE1>)
 {
     $line = $_;
+
+	$tempvalue = index($line, "//~ APIDOCTOOL:");
+	if ($tempvalue >= 0)
+	{
+		$lcline = lc($line)
+		$tempvalue = index($lcline, "document=");
+		if ($tempvalue >= 0)
+		{
+			# Skip ahead of "Document=" part.
+			$token = substr($lcline, $tempvalue + 9);
+			if (index($token, "on") eq 0)
+			{
+				$documenton = 1;
+			}
+			elsif (index($token, "off") eq 0)
+			{
+				$documenton = 0;
+			}
+			# else { Warning/error? }
+		}
+	}
+	
+	if (!$documenton)
+	{
+		print " \n";
+		next;
+	}
 	
 	if ($line =~ /\/\/ Copyright .+ Epic Games, Inc. All Rights Reserved./)
 	{
 		# Discard copyright line.
 		$line = " \n";
 	}
+	elsif ((index($line, "UCLASS(") eq 0) or (index($line, "UENUM(") eq 0) or (index($line, "USTRUCT(") eq 0) or (index($line, "UFUNCTION(") >= 0) or (index($line, "UPROPERTY(") >= 0) or (index($line, "GENERATED_BODY(") >= 0))
+	{
+		# Discard these Unreal macros - they confuse Doxygen.
+		$line = " \n";
+	}
 	else
 	{
 		if (index($line, "class ") eq 0)
@@ -37,25 +73,22 @@ while (<FILE1>)
 				{
 					$lastclassname = substr($line, $tempvalue);
 				}
-#print "lcn0 = " . $lastclassname . "\n";
 			}
 			elsif (index($line, ":") >= 0)
 			{
 				$tempvalue = index($line, "class ") + 6;
 				$lastclassname = substr($line, $tempvalue, index($line, ":") - $tempvalue);
-#print "lcn1 = " . $lastclassname . "\n";
 			}
 			else
 			{
 				# This will catch forward-declared classes, but should not be a problem because these should never immediately precede a GENERATED_UCLASS_BODY() line.
 				$lastclassname = substr($line, index($line, "class ") + 6);
-#print "lcn2 = " . $lastclassname . "\n";
 			}
 		}
 		elsif (index($line, "GENERATED_UCLASS_BODY()") >= 0)
 		{
 			# $lastclassname might not be initialized. If so, the previous logic block failed to catch a class name, or "GENERATED_UCLASS_BODY()" is floating around loose in a file somewhere.
-			$line = "GENERATED_BODY() public: " . $lastclassname . "(const FObjectInitializer& ObjectInitializer);\n";
+			$line = "public: " . $lastclassname . "(const FObjectInitializer& ObjectInitializer);\n";
 		}
 		elsif (index($line, "enum ") >= 0)
 		{
@@ -85,12 +118,6 @@ while (<FILE1>)
 					substr($line, $commentposition + 1, 0) = substr($line, $commentposition + 1, 1);
 				}
 			}
-			# This elsif indicates an error - the line didn't end with "\n".
-			# If we find this happening, we should start appending "\n" on our own.
-			#elsif ((length $line) > ($commentposition + 1))
-			#{
-			#    substr($line, $commentposition + 1, 0) = substr($line, $commentposition + 1, 1);
-			#}
 		}
 	}
 	print $line;
