Doxygen Perl script error fix.

diff --git a/Engine/Source/Programs/UnrealDocTool/DoxygenInputFilter/DoxygenInputFilter.pl b/Engine/Source/Programs/UnrealDocTool/DoxygenInputFilter/DoxygenInputFilter.pl
--- a/Engine/Source/Programs/UnrealDocTool/DoxygenInputFilter/DoxygenInputFilter.pl
+++ b/Engine/Source/Programs/UnrealDocTool/DoxygenInputFilter/DoxygenInputFilter.pl
@@ -24,7 +24,7 @@ while (<FILE1>)
 	$tempvalue = index($line, "//~ APIDOCTOOL:");
 	if ($tempvalue >= 0)
 	{
-		$lcline = lc($line)
+		$lcline = lc($line);
 		$tempvalue = index($lcline, "document=");
 		if ($tempvalue >= 0)
 		{
