Further Unreal macro filtering to help Doxygen avoid some newly-discovered warnings (APIDocTool/Doxygen).

diff --git a/Engine/Source/Programs/UnrealDocTool/DoxygenInputFilter/DoxygenInputFilter.pl b/Engine/Source/Programs/UnrealDocTool/DoxygenInputFilter/DoxygenInputFilter.pl
--- a/Engine/Source/Programs/UnrealDocTool/DoxygenInputFilter/DoxygenInputFilter.pl
+++ b/Engine/Source/Programs/UnrealDocTool/DoxygenInputFilter/DoxygenInputFilter.pl
@@ -53,7 +53,7 @@ while (<FILE1>)
 		# Discard copyright line.
 		$line = " \n";
 	}
-	elsif ((index($line, "UCLASS(") eq 0) or (index($line, "UENUM(") eq 0) or (index($line, "USTRUCT(") eq 0) or (index($line, "UFUNCTION(") >= 0) or (index($line, "UPROPERTY(") >= 0) or (index($line, "GENERATED_BODY(") >= 0))
+	elsif ((index($line, "UCLASS(") eq 0) or (index($line, "UENUM(") eq 0) or (index($line, "USTRUCT(") eq 0) or (index($line, "UFUNCTION(") >= 0) or (index($line, "UPROPERTY(") >= 0) or (index($line, "GENERATED_BODY(") >= 0) or (index($line, "GENERATED_USTRUCT_BODY(") >= 0))
 	{
 		# Discard these Unreal macros - they confuse Doxygen.
 		$line = " \n";
